Copy up from Dev-Editor [REDACTED_HANDLE]

diff --git a/Engine/Plugins/Experimental/MeshModelingToolset/Source/ModelingComponents/Private/SimpleDynamicMeshComponent.cpp b/Engine/Plugins/Experimental/MeshModelingToolset/Source/ModelingComponents/Private/SimpleDynamicMeshComponent.cpp
--- a/Engine/Plugins/Experimental/MeshModelingToolset/Source/ModelingComponents/Private/SimpleDynamicMeshComponent.cpp
+++ b/Engine/Plugins/Experimental/MeshModelingToolset/Source/ModelingComponents/Private/SimpleDynamicMeshComponent.cpp
@@ -129,6 +129,12 @@ FMeshTangentsf* USimpleDynamicMeshComponent::GetTangents()
 
 	// in this mode we assume the tangents are valid
 	check(TangentsType == EDynamicMeshTangentCalcType::ExternallyCalculated);
+	if (TangentsType == EDynamicMeshTangentCalcType::ExternallyCalculated)
+	{
+		// if you hit this, you did not request ExternallyCalculated tangents before initializing this PreviewMesh
+		check(Tangents.GetTangents().Num() > 0)
+	}
+
 	return &Tangents;
 }
 
@@ -145,8 +151,8 @@ void USimpleDynamicMeshComponent::NotifyMeshUpdated()
 {
 	// Need to recreate scene proxy to send it over
 	MarkRenderStateDirty();
+	LocalBounds = Mesh->GetCachedBounds();
 	UpdateBounds();
-	CurrentProxy = nullptr;
 
 	if (TangentsType != EDynamicMeshTangentCalcType::ExternallyCalculated)
 	{
@@ -157,9 +163,22 @@ void USimpleDynamicMeshComponent::NotifyMeshUpdated()
 
 void USimpleDynamicMeshComponent::FastNotifyColorsUpdated()
 {
-	if (CurrentProxy != nullptr)
+	FSimpleDynamicMeshSceneProxy* Proxy = GetCurrentSceneProxy();
+	if (Proxy != nullptr)
 	{
-		CurrentProxy->FastUpdateColors();
+		if (TriangleColorFunc != nullptr &&  Proxy->bUsePerTriangleColor == false )
+		{
+			Proxy->bUsePerTriangleColor = true;
+			Proxy->PerTriangleColorFunc = [this](const FDynamicMesh3* MeshIn, int TriangleID) { return GetTriangleColor(MeshIn, TriangleID); };
+		} 
+		else if (TriangleColorFunc == nullptr && Proxy->bUsePerTriangleColor == true)
+		{
+			Proxy->bUsePerTriangleColor = false;
+			Proxy->PerTriangleColorFunc = nullptr;
+		}
+
+		Proxy->FastUpdateVertices(false, false, true);
+		//MarkRenderDynamicDataDirty();
 	}
 	else
 	{
@@ -171,15 +190,19 @@ void USimpleDynamicMeshComponent::FastNotifyColorsUpdated()
 
 void USimpleDynamicMeshComponent::FastNotifyPositionsUpdated()
 {
-	if (CurrentProxy != nullptr)
+	if (GetCurrentSceneProxy() != nullptr)
 	{
-		bTangentsValid = false;
-		CurrentProxy->FastUpdatePositions(false);
+		GetCurrentSceneProxy()->FastUpdateVertices(true, false, false);
+		//MarkRenderDynamicDataDirty();
+		MarkRenderTransformDirty();
+		LocalBounds = Mesh->GetCachedBounds();
+		UpdateBounds();
 	}
 	else
 	{
 		NotifyMeshUpdated();
 	}
+
 }
 
 
@@ -187,33 +210,64 @@ void USimpleDynamicMeshComponent::FastNotifyPositionsUpdated()
 
 FPrimitiveSceneProxy* USimpleDynamicMeshComponent::CreateSceneProxy()
 {
-	CurrentProxy = nullptr;
+	check(GetCurrentSceneProxy() == nullptr);
+
+	FSimpleDynamicMeshSceneProxy* NewProxy = nullptr;
 	if (Mesh->TriangleCount() > 0)
 	{
-		CurrentProxy = new FSimpleDynamicMeshSceneProxy(this);
+		NewProxy = new FSimpleDynamicMeshSceneProxy(this);
 
 		if (TriangleColorFunc)
 		{
-			CurrentProxy->bUsePerTriangleColor = true;
-			CurrentProxy->PerTriangleColorFunc = [this](int TriangleID) { return GetTriangleColor(TriangleID); };
+			NewProxy->bUsePerTriangleColor = true;
+			NewProxy->PerTriangleColorFunc = [this](const FDynamicMesh3* MeshIn, int TriangleID) { return GetTriangleColor(MeshIn, TriangleID); };
 		}
 
-		CurrentProxy->Initialize();
+		if (SecondaryTriFilterFunc)
+		{
+			NewProxy->bUseSecondaryTriBuffers = true;
+			NewProxy->SecondaryTriFilterFunc = [this](const FDynamicMesh3* MeshIn, int32 TriangleID) 
+			{ 
+				return (SecondaryTriFilterFunc) ? SecondaryTriFilterFunc(MeshIn, TriangleID) : false;
+			};
+		}
+
+		NewProxy->Initialize();
 	}
-	return CurrentProxy;
+	return NewProxy;
 }
 
-int32 USimpleDynamicMeshComponent::GetNumMaterials() const
+
+void USimpleDynamicMeshComponent::NotifyMaterialSetUpdated()
 {
-	return 1;
+	if (GetCurrentSceneProxy() != nullptr)
+	{
+		GetCurrentSceneProxy()->UpdatedReferencedMaterials();
+	}
 }
 
 
-FColor USimpleDynamicMeshComponent::GetTriangleColor(int TriangleID)
+
+
+void USimpleDynamicMeshComponent::EnableSecondaryTriangleBuffers(TUniqueFunction<bool(const FDynamicMesh3*, int32)> SecondaryTriFilterFuncIn)
+{
+	SecondaryTriFilterFunc = MoveTemp(SecondaryTriFilterFuncIn);
+	NotifyMeshUpdated();
+}
+
+void USimpleDynamicMeshComponent::DisableSecondaryTriangleBuffers()
+{
+	SecondaryTriFilterFunc = nullptr;
+	NotifyMeshUpdated();
+}
+
+
+
+FColor USimpleDynamicMeshComponent::GetTriangleColor(const FDynamicMesh3* MeshIn, int TriangleID)
 {
 	if (TriangleColorFunc)
 	{
-		return TriangleColorFunc(TriangleID);
+		return TriangleColorFunc(MeshIn, TriangleID);
 	}
 	else
 	{
@@ -225,23 +279,15 @@ FColor USimpleDynamicMeshComponent::GetTriangleColor(int TriangleID)
 
 FBoxSphereBounds USimpleDynamicMeshComponent::CalcBounds(const FTransform& LocalToWorld) const
 {
-	// Bounds are tighter if the box is generated from pre-transformed vertices.
-	FBox BoundingBox(ForceInit);
-	for ( FVector3d Vertex : Mesh->VerticesItr() ) 
-	{
-		BoundingBox += LocalToWorld.TransformPosition(Vertex);
-	}
-
-	FBoxSphereBounds NewBounds;
-	NewBounds.BoxExtent = BoundingBox.GetExtent();
-	NewBounds.Origin = BoundingBox.GetCenter();
-	NewBounds.SphereRadius = NewBounds.BoxExtent.Size();
-
-	return NewBounds;
+	// can get a tighter box by calculating in world space, but we care more about performance
+	FBox LocalBoundingBox = (FBox)LocalBounds;
+	FBoxSphereBounds Ret(LocalBoundingBox.TransformBy(LocalToWorld));
+	Ret.BoxExtent *= BoundsScale;
+	Ret.SphereRadius *= BoundsScale;
+	return Ret;
 }
 
 
-
 void USimpleDynamicMeshComponent::ApplyChange(const FMeshVertexChange* Change, bool bRevert)
 {
 	int NV = Change->Vertices.Num();
@@ -266,4 +312,14 @@ void USimpleDynamicMeshComponent::ApplyChange(const FMeshChange* Change, bool bR
 
 	NotifyMeshUpdated();
 	OnMeshChanged.Broadcast();
-}
\ No newline at end of file
+}
+
+
+void USimpleDynamicMeshComponent::ApplyChange(const FMeshReplacementChange* Change, bool bRevert)
+{
+	Mesh->Copy(*Change->GetMesh(bRevert));
+
+	NotifyMeshUpdated();
+	OnMeshChanged.Broadcast();
+}
+
diff --git a/Engine/Plugins/Experimental/MeshModelingToolset/Source/ModelingComponents/Private/ToolSetupUtil.cpp b/Engine/Plugins/Experimental/MeshModelingToolset/Source/ModelingComponents/Private/ToolSetupUtil.cpp
--- a/Engine/Plugins/Experimental/MeshModelingToolset/Source/ModelingComponents/Private/ToolSetupUtil.cpp
+++ b/Engine/Plugins/Experimental/MeshModelingToolset/Source/ModelingComponents/Private/ToolSetupUtil.cpp
@@ -5,7 +5,7 @@
 #include "InteractiveTool.h"
 #include "InteractiveToolManager.h"
 #include "Materials/Material.h"
-
+#include "Materials/MaterialInstanceDynamic.h"
 
 
 UMaterialInterface* ToolSetupUtil::GetDefaultMaterial(UInteractiveToolManager* ToolManager, UMaterialInterface* SourceMaterial)
@@ -62,3 +62,36 @@ UMaterialInterface* ToolSetupUtil::GetSelectionMaterial(UInteractiveToolManager*
 	}
 	return Material;
 }
+
+
+UMaterialInterface* ToolSetupUtil::GetSelectionMaterial(const FLinearColor& UseColor, UInteractiveToolManager* ToolManager)
+{
+	check(ToolManager != nullptr);		// required for outer
+	UMaterialInterface* Material = LoadObject<UMaterial>(nullptr, TEXT("/MeshModelingToolset/Materials/SelectionMaterial"));
+	if (Material == nullptr && ToolManager != nullptr)
+	{
+		return ToolManager->GetContextQueriesAPI()->GetStandardMaterial(EStandardToolContextMaterials::VertexColorMaterial);
+	}
+	if (Material != nullptr)
+	{
+		UMaterialInstanceDynamic* MatInstance = UMaterialInstanceDynamic::Create(Material, ToolManager);
+		MatInstance->SetVectorParameterValue(TEXT("ConstantColor"), UseColor);
+		return MatInstance;
+	}
+	return Material;
+}
+
+
+
+
+UMaterialInterface* ToolSetupUtil::GetDefaultPointComponentMaterial(bool bRoundPoints, UInteractiveToolManager* ToolManager)
+{
+	UMaterialInterface* Material = (bRoundPoints) ?
+		LoadObject<UMaterial>(nullptr, TEXT("/MeshModelingToolset/Materials/PointSetComponentMaterialSoft")) :
+		LoadObject<UMaterial>(nullptr, TEXT("/MeshModelingToolset/Materials/PointSetComponentMaterial"));
+	if (Material == nullptr && ToolManager != nullptr)
+	{
+		return ToolManager->GetContextQueriesAPI()->GetStandardMaterial(EStandardToolContextMaterials::VertexColorMaterial);
+	}
+	return Material;
+}
\ No newline at end of file
diff --git a/Engine/Plugins/Experimental/MeshModelingToolset/Source/ModelingComponents/Public/Changes/ValueWatcher.h b/Engine/Plugins/Experimental/MeshModelingToolset/Source/ModelingComponents/Public/Changes/ValueWatcher.h
--- a/Engine/Plugins/Experimental/MeshModelingToolset/Source/ModelingComponents/Public/Changes/ValueWatcher.h
+++ b/Engine/Plugins/Experimental/MeshModelingToolset/Source/ModelingComponents/Public/Changes/ValueWatcher.h
@@ -41,4 +41,12 @@ public:
 			OnValueChangedFunc(CachedValue);
 		}
 	}
+
+	/**
+	 * Update known value without calling OnValueChangedFunc. Sometimes necessary during undo/redo.
+	 */
+	void SilentUpdate()
+	{
+		CachedValue = GetValueFunc();
+	}
 };
diff --git a/Engine/Plugins/Experimental/MeshModelingToolset/Source/ModelingComponents/Public/PreviewMesh.h b/Engine/Plugins/Experimental/MeshModelingToolset/Source/ModelingComponents/Public/PreviewMesh.h
--- a/Engine/Plugins/Experimental/MeshModelingToolset/Source/ModelingComponents/Public/PreviewMesh.h
+++ b/Engine/Plugins/Experimental/MeshModelingToolset/Source/ModelingComponents/Public/PreviewMesh.h
@@ -12,6 +12,7 @@
 
 class FDynamicMesh3;
 struct FMeshDescription;
+class FTangentsf;
 
 
 
@@ -52,7 +53,7 @@ public:
  * 
  */
 UCLASS(Transient)
-class MODELINGCOMPONENTS_API UPreviewMesh : public UObject, public IMeshVertexCommandChangeTarget, public IMeshCommandChangeTarget
+class MODELINGCOMPONENTS_API UPreviewMesh : public UObject, public IMeshVertexCommandChangeTarget, public IMeshCommandChangeTarget, public IMeshReplacementCommandChangeTarget
 {
 	GENERATED_BODY()
 
@@ -95,10 +96,74 @@ public:
 	 */
 	void SetMaterial(UMaterialInterface* Material);
 
+	/**
+	 * Set material on the preview mesh
+	 */
+	void SetMaterial(int MaterialIndex, UMaterialInterface* Material);
+
+	/**
+	 * Set the entire material set on the preview mesh
+	 */
+	void SetMaterials(const TArray<UMaterialInterface*>& Materials);
+
 	/**
 	* Get material from the preview mesh
 	*/
-	UMaterialInterface* GetMaterial() const;
+	UMaterialInterface* GetMaterial(int MaterialIndex = 0) const;
+
+	/**
+	 * Set an override material for the preview mesh. This material will override all the given materials.
+	 */
+	void SetOverrideRenderMaterial(UMaterialInterface* Material);
+
+	/**
+	 * Clear the override material for the preview mesh.
+	 */
+	void ClearOverrideRenderMaterial();
+
+	/**
+	 * @return the actual material that will be used for rendering for the given MaterialIndex. Will return override material if set.
+	 * 
+	 */
+	UMaterialInterface* GetActiveMaterial(int MaterialIndex = 0) const;
+
+
+	/**
+	 * Set an secondary material for the preview mesh. This material will be applied to secondary triangle buffer if enabled.
+	 */
+	void SetSecondaryRenderMaterial(UMaterialInterface* Material);
+
+	/**
+	 * Clear the secondary material for the preview mesh.
+	 */
+	void ClearSecondaryRenderMaterial();
+
+
+	/**
+	 * Enable secondary triangle buffers. The Secondary material will be applied to any triangles that pass TriangleFilterFunc.
+	 * @param TriangleFilterFunc predicate used to identify secondary triangles
+	 */
+	void EnableSecondaryTriangleBuffers(TUniqueFunction<bool(const FDynamicMesh3*, int32)> TriangleFilterFunc);
+
+	/**
+	 * Disable secondary triangle buffers
+	 */
+	void DisableSecondaryTriangleBuffers();
+
+
+
+	/**
+	 * Set the tangents mode for the underlying component, if available. 
+	 * Note that this function may need to be called before the mesh is initialized.
+	 */
+	void SetTangentsMode(EDynamicMeshTangentCalcType TangentsType);
+
+	/**
+	 * @return a MeshTangents data structure for the underlying component, if available, otherwise nullptr
+	 */
+	FMeshTangentsf* GetTangents() const;
+
+
 	/**
 	 * Get the current transform on the preview mesh
 	 */
@@ -134,7 +199,7 @@ public:
 	/**
 	 * Set the triangle color function for rendering / render data construction
 	 */
-	void SetTriangleColorFunction(TFunction<FColor(int)> TriangleColorFunc, ERenderUpdateMode UpdateMode = ERenderUpdateMode::FullUpdate);
+	void SetTriangleColorFunction(TFunction<FColor(const FDynamicMesh3*, int)> TriangleColorFunc, ERenderUpdateMode UpdateMode = ERenderUpdateMode::FullUpdate);
 
 	/**
 	 * Clear the triangle color function for rendering / render data construction
@@ -189,10 +254,15 @@ public:
 	 */
 	void InitializeMesh(FMeshDescription* MeshDescription);
 
+	/**
+	* @return pointer to the current FDynamicMesh used for preview  @todo deprecate this function, use GetMesh() instead
+	*/
+	const FDynamicMesh3* GetPreviewDynamicMesh() const { return GetMesh(); }
+
 	/**
 	* @return pointer to the current FDynamicMesh used for preview
 	*/
-	const FDynamicMesh3* GetPreviewDynamicMesh() const;
+	const FDynamicMesh3* GetMesh() const;
 
 	/**
 	 * @return the current preview FDynamicMesh, and replace with a new empty mesh
@@ -217,6 +287,24 @@ public:
 	 */
 	void EditMesh(TFunctionRef<void(FDynamicMesh3&)> EditFunc);
 
+	/**
+	 * Apply EditFunc to the internal mesh, and update spatial data structure if requested,
+	 * but do not update/rebuild rendering data structures. NotifyDeferredEditOcurred() must be
+	 * called to complete a deferred edit, this will update the rendering mesh.
+	 * DeferredEditMesh can be called multiple times before NotifyDeferredEditCompleted() is called.
+	 * @param EditFunc function that is applied to the internal mes
+	 * @param bRebuildSpatial if true, and internal spatial data structure is enabled, rebuild it for updated mesh
+	 */
+	void DeferredEditMesh(TFunctionRef<void(FDynamicMesh3&)> EditFunc, bool bRebuildSpatial);
+
+	/**
+	 * Notify that a DeferredEditMesh sequence is complete and cause update of rendering data structures.
+	 * @param UpdateMode type of rendering update required for the applied mesh edits
+	 * @param bRebuildSpatial if true, and internal spatial data structure is enabled, rebuild it for updated mesh
+	 */
+	void NotifyDeferredEditCompleted(ERenderUpdateMode UpdateMode, bool bRebuildSpatial);
+
+
 	/**
 	 * Apply EditFunc to the internal mesh and update internal data structures as necessary.
 	 * EditFunc is required to notify the given FDynamicMeshChangeTracker about all mesh changes
@@ -234,10 +322,22 @@ public:
 	 */
 	virtual void ApplyChange(const FMeshChange* Change, bool bRevert) override;
 
+	/**
+	* Apply/Revert a general mesh change to the internal mesh   (implements IMeshReplacementCommandChangeTarget)
+	*/
+	virtual void ApplyChange(const FMeshReplacementChange* Change, bool bRevert) override;
+
 	/** @return delegate that is broadcast whenever the internal mesh component is changed */
 	FSimpleMulticastDelegate& GetOnMeshChanged();
 
 
+	/**
+	 * Force rebuild of internal spatial data structure. Can be used in context of DeferredEditMesh to rebuild spatial data
+	 * structure w/o rebuilding render data
+	 */
+	void ForceRebuildSpatial();
+
+
 
 
 public:
diff --git a/Engine/Plugins/Experimental/MeshModelingToolset/Source/ModelingComponents/Public/Selection/GroupTopologySelector.h b/Engine/Plugins/Experimental/MeshModelingToolset/Source/ModelingComponents/Public/Selection/GroupTopologySelector.h
--- a/Engine/Plugins/Experimental/MeshModelingToolset/Source/ModelingComponents/Public/Selection/GroupTopologySelector.h
+++ b/Engine/Plugins/Experimental/MeshModelingToolset/Source/ModelingComponents/Public/Selection/GroupTopologySelector.h
@@ -10,19 +10,6 @@
 class FToolDataVisualizer;
 struct FViewCameraState;
 
-/**
- * FGroupTopologySelection represents a set of selected elements of a FGroupTopology
- */
-struct MODELINGCOMPONENTS_API FGroupTopologySelection
-{
-	TArray<int> SelectedGroupIDs;
-	TArray<int> SelectedCornerIDs;
-	TArray<int> SelectedEdgeIDs;
-
-	FGroupTopologySelection() { Clear(); }
-	void Clear();
-};
-
 
 /**
  * FGroupTopologySelector implements selection behavior for a FGroupTopology mesh.
diff --git a/Engine/Plugins/Experimental/MeshModelingToolset/Source/ModelingComponents/Public/Selection/SelectClickedAction.h b/Engine/Plugins/Experimental/MeshModelingToolset/Source/ModelingComponents/Public/Selection/SelectClickedAction.h
--- a/Engine/Plugins/Experimental/MeshModelingToolset/Source/ModelingComponents/Public/Selection/SelectClickedAction.h
+++ b/Engine/Plugins/Experimental/MeshModelingToolset/Source/ModelingComponents/Public/Selection/SelectClickedAction.h
@@ -11,7 +11,7 @@
  */
 class FSelectClickedAction : public IClickBehaviorTarget
 {
-	bool DoRayCast(const FInputDeviceRay& ClickPos, bool callbackOnHit)
+	FInputRayHit DoRayCast(const FInputDeviceRay& ClickPos, bool callbackOnHit)
 	{
 		FVector RayStart = ClickPos.WorldRay.Origin;
 		FVector RayEnd = ClickPos.WorldRay.PointAt(999999);
@@ -22,14 +22,14 @@ class FSelectClickedAction : public IClickBehaviorTarget
 		{
 			OnClickedPositionFunc(Result);
 		}
-		return bHitWorld;
+		return (bHitWorld) ? FInputRayHit(Result.Distance) : FInputRayHit();
 	}
 
 public:
 	UWorld* World;
 	TFunction<void(const FHitResult&)> OnClickedPositionFunc = nullptr;
 
-	virtual bool IsHitByClick(const FInputDeviceRay& ClickPos) override
+	virtual FInputRayHit IsHitByClick(const FInputDeviceRay& ClickPos) override
 	{
 		return DoRayCast(ClickPos, false);
 	}
diff --git a/Engine/Plugins/Experimental/MeshModelingToolset/Source/ModelingComponents/Public/SimpleDynamicMeshComponent.h b/Engine/Plugins/Experimental/MeshModelingToolset/Source/ModelingComponents/Public/SimpleDynamicMeshComponent.h
--- a/Engine/Plugins/Experimental/MeshModelingToolset/Source/ModelingComponents/Public/SimpleDynamicMeshComponent.h
+++ b/Engine/Plugins/Experimental/MeshModelingToolset/Source/ModelingComponents/Public/SimpleDynamicMeshComponent.h
@@ -61,6 +61,8 @@ public:
 	FMeshTangentsf* GetTangents();
 
 
+
+
 	/**
 	 * Write the internal mesh to a MeshDescription
 	 * @param bHaveModifiedTopology if false, we only update the vertex positions in the MeshDescription, otherwise it is Empty()'d and regenerated entirely
@@ -84,11 +86,23 @@ public:
 	//
 
 	/**
-	 * Call this if you update the mesh via GetMesh()
+	 * Call this if you update the mesh via GetMesh(). This will destroy the existing RenderProxy and create a new one.
 	 * @todo should provide a function that calls a lambda to modify the mesh, and only return const mesh pointer
 	 */
 	virtual void NotifyMeshUpdated() override;
 
+	/**
+	 * Call this instead of NotifyMeshUpdated() if you have only updated the vertex colors (or triangle color function).
+	 * This function will update the existing RenderProxy buffers if possible
+	 */
+	void FastNotifyColorsUpdated();
+
+	/**
+	 * Call this instead of NotifyMeshUpdated() if you have only updated the vertex positions.
+	 * This function will update the existing RenderProxy buffers if possible
+	 */
+	void FastNotifyPositionsUpdated();
+
 	/**
 	 * Apply a vertex deformation change to the internal mesh
 	 */
@@ -99,6 +113,11 @@ public:
 	 */
 	virtual void ApplyChange(const FMeshChange* Change, bool bRevert) override;
 
+	/**
+	* Apply a general mesh replacement change to the internal mesh
+	*/
+	virtual void ApplyChange(const FMeshReplacementChange* Change, bool bRevert) override;
+
 
 	/**
 	 * This delegate fires when a FCommandChange is applied to this component, so that
@@ -107,37 +126,62 @@ public:
 	FSimpleMulticastDelegate OnMeshChanged;
 
 
-	void SetDrawOnTop(bool bSet);
-
 	/**
 	 * if true, we always show the wireframe on top of the shaded mesh, even when not in wireframe mode
 	 */
 	UPROPERTY()
 	bool bExplicitShowWireframe = false;
 
+	/**
+	 * @return true if wireframe rendering pass is enabled
+	 */
+	virtual bool EnableWireframeRenderPass() const override { return bExplicitShowWireframe; }
+
+
+	/**
+	 * If this function is set, we will use these colors instead of vertex colors
+	 */
+	TFunction<FColor(const FDynamicMesh3*, int)> TriangleColorFunc = nullptr;
 
+
+	/**
+	 * If Secondary triangle buffers are enabled, then we will filter triangles that pass the given predicate
+	 * function into a second index buffer. These triangles will be drawn with the Secondary render material
+	 * that is set in the BaseDynamicMeshComponent. Calling this function invalidates the SceneProxy.
+	 */
+	virtual void EnableSecondaryTriangleBuffers(TUniqueFunction<bool(const FDynamicMesh3*, int32)> SecondaryTriFilterFunc);
+
+	/**
+	 * Disable secondary triangle buffers. This invalidates the SceneProxy.
+	 */
+	virtual void DisableSecondaryTriangleBuffers();
+
+
+public:
+
+	// do not use this
 	UPROPERTY()
 	bool bDrawOnTop = false;
 
+	// do not use this
+	void SetDrawOnTop(bool bSet);
 
 
-	TFunction<FColor(int)> TriangleColorFunc = nullptr;
-	void FastNotifyColorsUpdated();
-
-	void FastNotifyPositionsUpdated();
+protected:
+	/**
+	 * This is called to tell our RenderProxy about modifications to the material set.
+	 * We need to pass this on for things like material validation in the Editor.
+	 */
+	virtual void NotifyMaterialSetUpdated();
 
 private:
 
-	FSimpleDynamicMeshSceneProxy* CurrentProxy = nullptr;
+	FSimpleDynamicMeshSceneProxy* GetCurrentSceneProxy() { return (FSimpleDynamicMeshSceneProxy*)SceneProxy; }
 
 	//~ Begin UPrimitiveComponent Interface.
 	virtual FPrimitiveSceneProxy* CreateSceneProxy() override;
 	//~ End UPrimitiveComponent Interface.
 
-	//~ Begin UMeshComponent Interface.
-	virtual int32 GetNumMaterials() const override;
-	//~ End UMeshComponent Interface.
-
 	//~ Begin USceneComponent Interface.
 	virtual FBoxSphereBounds CalcBounds(const FTransform& LocalToWorld) const override;
 	//~ Begin USceneComponent Interface.
@@ -145,10 +189,16 @@ private:
 	TUniquePtr<FDynamicMesh3> Mesh;
 	void InitializeNewMesh();
 
+	// local-space bounding of Mesh
+	FAxisAlignedBox3d LocalBounds;
+
 	bool bTangentsValid = false;
 	FMeshTangentsf Tangents;
 	
-	FColor GetTriangleColor(int TriangleID);
+	FColor GetTriangleColor(const FDynamicMesh3* Mesh, int TriangleID);
+
+	TUniqueFunction<bool(const FDynamicMesh3*, int32)> SecondaryTriFilterFunc = nullptr;
+
 
 	//friend class FCustomMeshSceneProxy;
 };
diff --git a/Engine/Plugins/Experimental/MeshModelingToolset/Source/ModelingComponents/Public/ToolSetupUtil.h b/Engine/Plugins/Experimental/MeshModelingToolset/Source/ModelingComponents/Public/ToolSetupUtil.h
--- a/Engine/Plugins/Experimental/MeshModelingToolset/Source/ModelingComponents/Public/ToolSetupUtil.h
+++ b/Engine/Plugins/Experimental/MeshModelingToolset/Source/ModelingComponents/Public/ToolSetupUtil.h
@@ -41,4 +41,17 @@ namespace ToolSetupUtil
 	 */
 	MODELINGCOMPONENTS_API UMaterialInterface* GetSelectionMaterial(UInteractiveToolManager* ToolManager);
 
+	/**
+	 * @return Selection Material 1 with custom color
+	 */
+	MODELINGCOMPONENTS_API UMaterialInterface* GetSelectionMaterial(const FLinearColor& UseColor, UInteractiveToolManager* ToolManager);
+
+
+
+	/**
+	 * @param bRoundPoints true for round points, false for square
+	 * @return custom material suitable for use with UPointSetComponent
+	 */
+	MODELINGCOMPONENTS_API UMaterialInterface* GetDefaultPointComponentMaterial(bool bRoundPoints, UInteractiveToolManager* ToolManager);
+
 }
\ No newline at end of file
diff --git a/Engine/Plugins/Experimental/MeshModelingToolset/Source/ModelingComponents/Public/Transforms/MultiTransformer.h b/Engine/Plugins/Experimental/MeshModelingToolset/Source/ModelingComponents/Public/Transforms/MultiTransformer.h
new file mode 100644
--- /dev/null
+++ b/Engine/Plugins/Experimental/MeshModelingToolset/Source/ModelingComponents/Public/Transforms/MultiTransformer.h
@@ -0,0 +1,88 @@
+// Copyright 1998-2019 Epic Games, Inc. All Rights Reserved.
+
+#pragma once
+
+#include "CoreMinimal.h"
+#include "ToolContextInterfaces.h"
+#include "ToolDataVisualizer.h"
+#include "FrameTypes.h"
+#include "InteractiveGizmo.h"
+#include "BaseGizmos/TransformGizmo.h"
+#include "InteractiveGizmoManager.h"
+#include "MultiTransformer.generated.h"
+
+
+UENUM()
+enum class EMultiTransformerMode
+{
+	DefaultGizmo = 1,
+	QuickAxisTranslation = 2
+};
+
+
+/**
+
+ */
+UCLASS()
+class MODELINGCOMPONENTS_API UMultiTransformer : public UObject
+{
+	GENERATED_BODY()
+public:
+	
+	virtual void Setup(UInteractiveGizmoManager* GizmoManager);
+	virtual void Shutdown();
+
+	virtual void Tick(float DeltaTime);
+
+	virtual void SetGizmoPositionFromWorldFrame(const FFrame3d& Frame);
+	virtual void SetGizmoPositionFromWorldPos(const FVector& Position, const FVector& Normal);
+
+	virtual const FFrame3d& GetCurrentGizmoFrame() const { return ActiveGizmoFrame; }
+	virtual bool InGizmoEdit() const { return bInGizmoEdit;	}
+
+	virtual EMultiTransformerMode GetMode() const { return ActiveMode; }
+	virtual void SetMode(EMultiTransformerMode NewMode);
+
+	virtual void SetGizmoVisibility(bool bVisible);
+
+	void SetSnapToWorldGridSourceFunc(TUniqueFunction<bool()> EnableSnapFunc);
+
+public:
+	DECLARE_MULTICAST_DELEGATE(FMultiTransformerEvent);
+
+	/** This delegate is fired when a drag is started */
+	FMultiTransformerEvent OnTransformStarted;
+
+	/** This delegate is fired when a drag is updated */
+	FMultiTransformerEvent OnTransformUpdated;
+
+	/** This delegate is fired when the drag is completed */
+	FMultiTransformerEvent OnTransformCompleted;
+
+
+public:
+	UInteractiveGizmoManager* GizmoManager;
+
+	EMultiTransformerMode ActiveMode;
+
+	bool bShouldBeVisible = true;
+	FFrame3d ActiveGizmoFrame;
+
+	UPROPERTY()
+	UTransformGizmo* TransformGizmo;
+
+	UPROPERTY()
+	UTransformProxy* TransformProxy;
+
+	TUniqueFunction<bool()> EnableSnapToWorldGridFunc;
+
+	// called on PlaneTransformProxy.OnTransformChanged
+	void OnProxyTransformChanged(UTransformProxy* Proxy, FTransform Transform);
+	void OnBeginProxyTransformEdit(UTransformProxy* Proxy);
+	void OnEndProxyTransformEdit(UTransformProxy* Proxy);
+	bool bInGizmoEdit = false;
+
+	void UpdateShowGizmoState(bool bNewVisibility);
+
+
+};
\ No newline at end of file
diff --git a/Engine/Plugins/Experimental/MeshModelingToolset/Source/ModelingOperators/Private/CleaningOps/EditNormalsOp.cpp b/Engine/Plugins/Experimental/MeshModelingToolset/Source/ModelingOperators/Private/CleaningOps/EditNormalsOp.cpp
--- a/Engine/Plugins/Experimental/MeshModelingToolset/Source/ModelingOperators/Private/CleaningOps/EditNormalsOp.cpp
+++ b/Engine/Plugins/Experimental/MeshModelingToolset/Source/ModelingOperators/Private/CleaningOps/EditNormalsOp.cpp
@@ -37,7 +37,7 @@ void FEditNormalsOp::CalculateResult(FProgressCancel* Progress)
 	}
 
 	// if you split normals you must always recompute as well
-	bool bNeedsRecompute = bRecomputeNormals || bSplitNormals;
+	bool bNeedsRecompute = bRecomputeNormals || SplitNormalMethod != ESplitNormalMethod::UseExistingTopology;
 
 	if (bFixInconsistentNormals)
 	{
@@ -84,14 +84,24 @@ void FEditNormalsOp::CalculateResult(FProgressCancel* Progress)
 	float NormalDotProdThreshold = FMathf::Cos(NormalSplitThreshold * FMathf::DegToRad);
 
 	FMeshNormals FaceNormals(ResultMesh.Get());
-	if (bSplitNormals)
+	if (SplitNormalMethod != ESplitNormalMethod::UseExistingTopology)
 	{
-		FaceNormals.ComputeTriangleNormals();
-		const TArray<FVector3d>& Normals = FaceNormals.GetNormals();
-		ResultMesh->Attributes()->PrimaryNormals()->CreateFromPredicate([&Normals, &NormalDotProdThreshold](int VID, int TA, int TB)
+		if (SplitNormalMethod == ESplitNormalMethod::FaceNormalThreshold)
 		{
-			return Normals[TA].Dot(Normals[TB]) > NormalDotProdThreshold;
-		}, 0);
+			FaceNormals.ComputeTriangleNormals();
+			const TArray<FVector3d>& Normals = FaceNormals.GetNormals();
+			ResultMesh->Attributes()->PrimaryNormals()->CreateFromPredicate([&Normals, &NormalDotProdThreshold](int VID, int TA, int TB)
+			{
+				return Normals[TA].Dot(Normals[TB]) > NormalDotProdThreshold;
+			}, 0);
+		}
+		else // SplitNormalMethod == ESplitNormalMethod::FaceGroupID
+		{
+			ResultMesh->Attributes()->PrimaryNormals()->CreateFromPredicate([this](int VID, int TA, int TB)
+			{
+				return ResultMesh->GetTriangleGroup(TA) == ResultMesh->GetTriangleGroup(TB);
+			}, 0);
+		}
 	}
 
 	if (Progress->Cancelled())
@@ -114,7 +124,7 @@ void FEditNormalsOp::CalculateResult(FProgressCancel* Progress)
 		return;
 	}
 
-	if (bAllowSharpVertices)
+	if (SplitNormalMethod == ESplitNormalMethod::FaceNormalThreshold && bAllowSharpVertices)
 	{
 		ResultMesh->Attributes()->PrimaryNormals()->SplitVerticesWithPredicate([this, &FaceNormals, &NormalDotProdThreshold](int ElementID, int TriID)
 		{
diff --git a/Engine/Plugins/Experimental/MeshModelingToolset/Source/ModelingOperators/Private/CuttingOps/PlaneCutOp.cpp b/Engine/Plugins/Experimental/MeshModelingToolset/Source/ModelingOperators/Private/CuttingOps/PlaneCutOp.cpp
--- a/Engine/Plugins/Experimental/MeshModelingToolset/Source/ModelingOperators/Private/CuttingOps/PlaneCutOp.cpp
+++ b/Engine/Plugins/Experimental/MeshModelingToolset/Source/ModelingOperators/Private/CuttingOps/PlaneCutOp.cpp
@@ -21,22 +21,64 @@ void FPlaneCutOp::CalculateResult(FProgressCancel* Progress)
 	{
 		return;
 	}
-	ResultMesh->Copy(*OriginalMesh, true, true, true, !bDiscardAttributes);
+	ResultMesh->Copy(*OriginalMesh, true, true, true, true);
 
+	CutPlaneLocalThickness = FMath::Abs(CutPlaneLocalThickness); // sanitize thickness param
+	
 	FMeshPlaneCut Cut(ResultMesh.Get(), LocalPlaneOrigin, LocalPlaneNormal);
 
 	if (Progress->Cancelled())
 	{
 		return;
 	}
-	Cut.UVScaleFactor = 1.0 / ResultMesh->GetBounds().MaxDim();
+	Cut.UVScaleFactor = UVScaleFactor;
 
 
 	if (Progress->Cancelled())
 	{
 		return;
 	}
-	Cut.Cut();
+	if (bKeepBothHalves)
+	{
+		int MaxSubObjectID = -1;
+		TDynamicMeshScalarTriangleAttribute<int>* SubObjectAttrib = static_cast<TDynamicMeshScalarTriangleAttribute<int>*>(ResultMesh->Attributes()->GetAttachedAttribute(SubObjectsAttribIndex));
+		for (int TID : ResultMesh->TriangleIndicesItr())
+		{
+			MaxSubObjectID = FMath::Max(MaxSubObjectID, SubObjectAttrib->GetValue(TID));
+		}
+		if (CutPlaneLocalThickness <= Cut.PlaneTolerance)
+		{
+			Cut.CutWithoutDelete(true, 0, SubObjectAttrib, MaxSubObjectID+1);
+		}
+		else // for a 'thick plane' we need two cuts
+		{
+			Cut.PlaneOrigin -= Cut.PlaneNormal * CutPlaneLocalThickness;
+			Cut.CutWithoutDelete(true, 0, SubObjectAttrib, MaxSubObjectID + 1, true, false);
+
+			int SecondCutMaxID = MaxSubObjectID;
+			for (int TID : ResultMesh->TriangleIndicesItr())
+			{
+				SecondCutMaxID = FMath::Max(SecondCutMaxID, SubObjectAttrib->GetValue(TID));
+			}
+			Cut.PlaneOrigin += Cut.PlaneNormal * (2 * CutPlaneLocalThickness);
+			Cut.CutWithoutDelete(true, 0, SubObjectAttrib, SecondCutMaxID + 1, false, true);
+
+			for (int TID : ResultMesh->TriangleIndicesItr())
+			{
+				int SubObjectID = SubObjectAttrib->GetValue(TID);
+				if (SubObjectID > MaxSubObjectID && SubObjectID <= SecondCutMaxID)
+				{
+					ResultMesh->RemoveTriangle(TID);
+				}
+			}
+		}
+
+		
+	}
+	else
+	{
+		Cut.Cut();
+	}
 
 
 	if (Progress->Cancelled())
@@ -47,4 +89,14 @@ void FPlaneCutOp::CalculateResult(FProgressCancel* Progress)
 	{
 		Cut.HoleFill(ConstrainedDelaunayTriangulate<double>, bFillSpans);
 	}
+
+	if (Progress->Cancelled())
+	{
+		return;
+	}
+	if (bFillCutHole && bKeepBothHalves)
+	{
+		TDynamicMeshScalarTriangleAttribute<int>* SubObjectAttrib = static_cast<TDynamicMeshScalarTriangleAttribute<int>*>(ResultMesh->Attributes()->GetAttachedAttribute(SubObjectsAttribIndex));
+		Cut.TransferTriangleLabelsToHoleFillTriangles(SubObjectAttrib);
+	}
 }
\ No newline at end of file
diff --git a/Engine/Plugins/Experimental/MeshModelingToolset/Source/ModelingOperators/Private/ParameterizationOps/UVProjectionOp.cpp b/Engine/Plugins/Experimental/MeshModelingToolset/Source/ModelingOperators/Private/ParameterizationOps/UVProjectionOp.cpp
--- a/Engine/Plugins/Experimental/MeshModelingToolset/Source/ModelingOperators/Private/ParameterizationOps/UVProjectionOp.cpp
+++ b/Engine/Plugins/Experimental/MeshModelingToolset/Source/ModelingOperators/Private/ParameterizationOps/UVProjectionOp.cpp
@@ -52,7 +52,7 @@ void FUVProjectionOp::CalculateResult(FProgressCancel* Progress)
 	FDynamicMeshUVOverlay* UVLayer = ResultMesh->Attributes()->GetUVLayer(LayerIndex);
 	UVLayer->ClearElements();
 
-	FVector2D Scale = FVector2D(.5,.5)*UVScale, Offset = FVector2D(.5,.5)+UVOffset;
+	FVector2f Scale = FVector2f(.5,.5)*UVScale, Offset = FVector2f(.5,.5)+UVOffset;
 	// project to major axis
 	auto ProjAxis = [&Scale, &Offset](const FVector3d& P, int Ax1, int Ax2, float Ax1Scale, float Ax2Scale)
 	{
@@ -70,7 +70,7 @@ void FUVProjectionOp::CalculateResult(FProgressCancel* Progress)
 			if (ResultMesh->IsVertex(VID))
 			{
 				FVector2f UV = ProjAxis(TransformedVertices[VID], 0, 1, 1, 1);
-				UVLayer->InsertElement(VID, UV, VID, true);
+				UVLayer->InsertElement(VID, (float*)UV, VID, true);
 			}
 		}
 		UVLayer->EndUnsafeElementsInsert();
diff --git a/Engine/Plugins/Experimental/MeshModelingToolset/Source/ModelingOperators/Public/CleaningOps/EditNormalsOp.h b/Engine/Plugins/Experimental/MeshModelingToolset/Source/ModelingOperators/Public/CleaningOps/EditNormalsOp.h
--- a/Engine/Plugins/Experimental/MeshModelingToolset/Source/ModelingOperators/Public/CleaningOps/EditNormalsOp.h
+++ b/Engine/Plugins/Experimental/MeshModelingToolset/Source/ModelingOperators/Public/CleaningOps/EditNormalsOp.h
@@ -18,6 +18,14 @@ enum class ENormalCalculationMethod : uint8
 };
 
 
+UENUM()
+enum class ESplitNormalMethod : uint8
+{
+	UseExistingTopology,
+	FaceNormalThreshold,
+	FaceGroupID
+};
+
 
 class MODELINGOPERATORS_API FEditNormalsOp : public FDynamicMeshOperator
 {
@@ -31,7 +39,7 @@ public:
 	bool bInvertNormals;
 	bool bRecomputeNormals;
 	ENormalCalculationMethod NormalCalculationMethod;
-	bool bSplitNormals;
+	ESplitNormalMethod SplitNormalMethod;
 	bool bAllowSharpVertices;
 	float NormalSplitThreshold;
 
diff --git a/Engine/Plugins/Experimental/MeshModelingToolset/Source/ModelingOperators/Public/CuttingOps/PlaneCutOp.h b/Engine/Plugins/Experimental/MeshModelingToolset/Source/ModelingOperators/Public/CuttingOps/PlaneCutOp.h
--- a/Engine/Plugins/Experimental/MeshModelingToolset/Source/ModelingOperators/Public/CuttingOps/PlaneCutOp.h
+++ b/Engine/Plugins/Experimental/MeshModelingToolset/Source/ModelingOperators/Public/CuttingOps/PlaneCutOp.h
@@ -17,8 +17,12 @@ public:
 
 	// inputs
 	FVector LocalPlaneOrigin, LocalPlaneNormal;
-	bool bFillCutHole, bDiscardAttributes, bFillSpans;
-	TSharedPtr<FDynamicMesh3> OriginalMesh;
+	bool bFillCutHole, bFillSpans;
+	bool bKeepBothHalves;
+	float CutPlaneLocalThickness = 0; // plane thickness in the local space of the mesh
+	float UVScaleFactor;
+	int SubObjectsAttribIndex;
+	TSharedPtr<const FDynamicMesh3> OriginalMesh;
 
 	void SetTransform(const FTransform& Transform);
 
diff --git a/Engine/Plugins/Experimental/MeshModelingToolset/Source/ModelingOperators/Public/ParameterizationOps/UVProjectionOp.h b/Engine/Plugins/Experimental/MeshModelingToolset/Source/ModelingOperators/Public/ParameterizationOps/UVProjectionOp.h
--- a/Engine/Plugins/Experimental/MeshModelingToolset/Source/ModelingOperators/Public/ParameterizationOps/UVProjectionOp.h
+++ b/Engine/Plugins/Experimental/MeshModelingToolset/Source/ModelingOperators/Public/ParameterizationOps/UVProjectionOp.h
@@ -30,7 +30,7 @@ public:
 	EUVProjectionMethod ProjectionMethod;
 	float CylinderProjectToTopOrBottomAngleThreshold;
 	FTransform ProjectionTransform;
-	FVector2D UVScale, UVOffset;
+	FVector2f UVScale, UVOffset;
 	bool bWorldSpaceUVScale;
 
 	void SetTransform(const FTransform& Transform);
diff --git a/Engine/Plugins/Experimental/MeshModelingToolset/Source/ModelingOperatorsEditorOnly/Private/ParameterizationOps/ParameterizeMeshOp.cpp b/Engine/Plugins/Experimental/MeshModelingToolset/Source/ModelingOperatorsEditorOnly/Private/ParameterizationOps/ParameterizeMeshOp.cpp
--- a/Engine/Plugins/Experimental/MeshModelingToolset/Source/ModelingOperatorsEditorOnly/Private/ParameterizationOps/ParameterizeMeshOp.cpp
+++ b/Engine/Plugins/Experimental/MeshModelingToolset/Source/ModelingOperatorsEditorOnly/Private/ParameterizationOps/ParameterizeMeshOp.cpp
@@ -5,11 +5,13 @@
 #include "DynamicMeshAttributeSet.h"
 #include "MeshDescriptionToDynamicMesh.h"
 #include "ProxyLODParameterization.h"
+#include "Selections/MeshConnectedComponents.h"
 
 
 
 
-FParameterizeMeshOp::FLinearMesh::FLinearMesh(const FDynamicMesh3& Mesh)
+
+FParameterizeMeshOp::FLinearMesh::FLinearMesh(const FDynamicMesh3& Mesh, const bool bRespectPolygroups)
 {
 
 	TArray<FVector>& Positions = this->VertexBuffer;
@@ -22,8 +24,8 @@ FParameterizeMeshOp::FLinearMesh::FLinearMesh(const FDynamicMesh3& Mesh)
 
 	// Compute the mapping from triangle ID to triangle number
 	{
-		const int32 MaxTriID = Mesh.MaxTriangleID();
-		const int32 NumTris = Mesh.TriangleCount();
+		const int32 MaxTriID = Mesh.MaxTriangleID(); // really +1.. all TriID < MaxTriID
+		const int32 NumTris  = Mesh.TriangleCount();
 
 		// reserve space and add elements 
 		TriFromID.Empty(MaxTriID);
@@ -44,7 +46,7 @@ FParameterizeMeshOp::FLinearMesh::FLinearMesh(const FDynamicMesh3& Mesh)
 	// Compute the mapping from vertex ID to vertex number
 	{
 		const int32 MaxVertID = Mesh.MaxVertexID();
-		const int32 NumVerts = Mesh.VertexCount();
+		const int32 NumVerts  = Mesh.VertexCount();
 
 		// reserve space and add elements
 		VertFromID.Empty(MaxVertID);
@@ -95,33 +97,71 @@ FParameterizeMeshOp::FLinearMesh::FLinearMesh(const FDynamicMesh3& Mesh)
 	// For each edge on each triangle.
 	AdjacencyBuffer.Empty(NumTris * 3);
 
-	for (int TriID : Mesh.TriangleIndicesItr())
+	// Create Adjacency - create boundaries in adjacency at polygroup boundaries if requested.
+	if (bRespectPolygroups && Mesh.HasTriangleGroups())
 	{
-
-		// NB: this maybe the same as GetNeighborTriangles()
-		for (int i = 0; i < 3; ++i)
+		for (int TriID : Mesh.TriangleIndicesItr())
 		{
-			int32 EdgesID = Mesh.GetTriEdge(TriID, i);
-			FIndex4i Edge = Mesh.GetEdge(EdgesID);
-			int32 OtherTriID = (Edge[2] == TriID) ? Edge[3] : Edge[2];
+			FIndex3i NbrTris = Mesh.GetTriNeighbourTris(TriID);
+			for (int32 i = 0; i < 3; ++i)
+			{
+				const int32 NbrID = NbrTris[i];
 
-			int32 RemapOtherID = (OtherTriID != FDynamicMesh3::InvalidID) ? TriFromID[OtherTriID] : FDynamicMesh3::InvalidID;
+				int32 RemapNbrID = FDynamicMesh3::InvalidID;
+				if (NbrID != FDynamicMesh3::InvalidID)
+				{
+					const int32 CurTriGroup = Mesh.GetTriangleGroup(TriID);
+					const int32 NbrTriGroup = Mesh.GetTriangleGroup(NbrID);
 
-			AdjacencyBuffer.Add(RemapOtherID);
+					RemapNbrID = (CurTriGroup == NbrTriGroup) ? TriFromID[NbrID] : FDynamicMesh3::InvalidID;
+				}
+
+				AdjacencyBuffer.Add(RemapNbrID);
+			}
 		}
+
 	}
+	else  // compute the adjacency using only the mesh connectivity.
+	{
+		for (int TriID : Mesh.TriangleIndicesItr())
+		{
+			FIndex3i NbrTris = Mesh.GetTriNeighbourTris(TriID);
+			for (int32 i = 0; i < 3; ++i)
+			{
+				int32 NbrID = NbrTris[i];
 
+				int32 RemapNbrID = (NbrID != FDynamicMesh3::InvalidID) ? TriFromID[NbrID] : FDynamicMesh3::InvalidID;
+				AdjacencyBuffer.Add(RemapNbrID);
+			}
+		}
+	}
 }
 
-bool FParameterizeMeshOp::ComputeUVs(FDynamicMesh3& Mesh, TFunction<bool(float)>& Interrupter)
+bool FParameterizeMeshOp::ComputeUVs(FDynamicMesh3& Mesh,  TFunction<bool(float)>& Interrupter, const bool bUsePolygroups, float GlobalScale)
 {
+	// the UVAltas code is unhappy if you feed it a single degenerate triangle
+	bool bNonDegenerate = true;
+	if (Mesh.TriangleCount() == 1)
+	{
+		for (auto TriID : Mesh.TriangleIndicesItr())
+		{
+			double Area = Mesh.GetTriArea(TriID);
+			bNonDegenerate = bNonDegenerate && FMath::Abs(Area) > 1.e-5;
+		}
+	}
+
+	if (!bNonDegenerate)
+	{
+		return false;
+	}
+
 	// Convert to a dense form.
-	FLinearMesh LinearMesh(Mesh);
+	FLinearMesh LinearMesh(Mesh, bUsePolygroups);
 
 	// Data to be populated by the UV generation tool
 	TArray<FVector2D> UVVertexBuffer;
 	TArray<int32>     UVIndexBuffer;
-	TArray<int32>     VertexRemapArray;
+	TArray<int32>     VertexRemapArray; // This maps the UV vertices to the original position vertices.  Note multiple UV vertices might share the same positional vertex (due to UV boundaries)
 
 
 	float MaxStretch     = Stretch;
@@ -156,6 +196,8 @@ bool FParameterizeMeshOp::ComputeUVs(FDynamicMesh3& Mesh, TFunction<bool(float)>
 
 		// Add the UVs to the overlay
 		int32 NumUVs = UVVertexBuffer.Num();
+		TArray<int32> UVOffsetToElID;  UVOffsetToElID.Reserve(NumUVs);
+
 		for (int32 i = 0; i < NumUVs; ++i)
 		{
 			const FVector2D UV = UVVertexBuffer[i];
@@ -165,8 +207,8 @@ bool FParameterizeMeshOp::ComputeUVs(FDynamicMesh3& Mesh, TFunction<bool(float)>
 			const int32 VertID = LinearMesh.VertToID[VertOffset];
 
 			// add the UV to the mesh overlay
-			const int32 NewID = UVOverlay->AppendElement(UV, VertID);
-
+			const int32 NewID = UVOverlay->AppendElement(FVector2f(UV), VertID);
+			UVOffsetToElID.Add(NewID);
 		}
 
 		int32 NumUVTris = UVIndexBuffer.Num() / 3;
@@ -174,18 +216,18 @@ bool FParameterizeMeshOp::ComputeUVs(FDynamicMesh3& Mesh, TFunction<bool(float)>
 		{
 			int32 t = i * 3;
 			// The triangle in UV space
-			FIndex3i TriUV(UVIndexBuffer[t], UVIndexBuffer[t + 1], UVIndexBuffer[t + 2]);
+			FIndex3i UVTri(UVIndexBuffer[t], UVIndexBuffer[t + 1], UVIndexBuffer[t + 2]);
 
 			// the triangle in terms of the VertIDs in the DynamicMesh
 			FIndex3i TriVertIDs;
 			for (int c = 0; c < 3; ++c)
 			{
 				// the offset for this vertex in the LinearMesh
-				int32 Offset = VertexRemapArray[TriUV[c]];
+				int32 Offset = VertexRemapArray[UVTri[c]];
 
-				int32 VertId = LinearMesh.VertToID[Offset];
+				int32 VertID = LinearMesh.VertToID[Offset];
 
-				TriVertIDs[c] = VertId;
+				TriVertIDs[c] = VertID;
 			}
 
 			// NB: this could be slow.. 
@@ -193,16 +235,76 @@ bool FParameterizeMeshOp::ComputeUVs(FDynamicMesh3& Mesh, TFunction<bool(float)>
 
 			checkSlow(TriID != FDynamicMesh3::InvalidID);
 
+			FIndex3i ElTri(UVOffsetToElID[ UVTri[0]], UVOffsetToElID[UVTri[1]], UVOffsetToElID[UVTri[2]]);
+
 			// add the triangle to the overlay
-			UVOverlay->SetTriangle(TriID, TriUV);
+			UVOverlay->SetTriangle(TriID, ElTri);
 		}
 
+		if (bNormalizeAreas)
+		{
+			NormalizeUVAreas(&Mesh, UVOverlay, GlobalScale);
+		}
 	}
 
+
 	return bSuccess;
 }
 
 
+
+
+void FParameterizeMeshOp::NormalizeUVAreas(const FDynamicMesh3* Mesh, FDynamicMeshUVOverlay* Overlay, float GlobalScale)
+{
+	FMeshConnectedComponents UVComponents(Mesh);
+	UVComponents.FindConnectedTriangles([&](int32 Triangle0, int32 Triangle1) {
+		return Overlay->AreTrianglesConnected(Triangle0, Triangle1);
+	});
+
+	// TODO ParallelFor
+	for (FMeshConnectedComponents::FComponent& Component : UVComponents)
+	{
+		const TArray<int>& Triangles = Component.Indices;
+		TSet<int> UVElements;
+		UVElements.Reserve(Triangles.Num() * 3);
+		double AreaUV = 0;
+		double Area3D = 0;
+		FVector3d Triangle3D[3];
+		FVector2f TriangleUV[3];
+		FAxisAlignedBox2f BoundsUV = FAxisAlignedBox2f::Empty();
+
+		for (int tid : Triangles)
+		{
+			FIndex3i TriElements = Overlay->GetTriangle(tid);
+			if (!TriElements.Contains(FDynamicMesh3::InvalidID))
+			{
+				for (int j = 0; j < 3; ++j)
+				{
+					TriangleUV[j] = Overlay->GetElement(TriElements[j]);
+					BoundsUV.Contain(TriangleUV[j]);
+					Triangle3D[j] = Mesh->GetVertex(Overlay->GetParentVertex(TriElements[j]));
+					UVElements.Add(TriElements[j]);
+				}
+				AreaUV += VectorUtil::Area(TriangleUV[0], TriangleUV[1], TriangleUV[2]);
+				Area3D += VectorUtil::Area(Triangle3D[0], Triangle3D[1], Triangle3D[2]);
+			}
+		}
+		
+		double LinearScale = (AreaUV > 0.00001) ? ( FMathd::Sqrt(Area3D) / FMathd::Sqrt(AreaUV)) : 1.0;
+		LinearScale = LinearScale * GlobalScale;
+		FVector2f ComponentOrigin = BoundsUV.Center();
+
+		for (int elemid : UVElements)
+		{
+			FVector2f UV = Overlay->GetElement(elemid);
+			UV = (UV - ComponentOrigin) * (float)LinearScale + ComponentOrigin;
+			Overlay->SetElement(elemid, UV);
+		}
+	}
+}
+
+
+
 void FParameterizeMeshOp::CalculateResult(FProgressCancel* Progress)
 {
 
@@ -227,10 +329,197 @@ void FParameterizeMeshOp::CalculateResult(FProgressCancel* Progress)
 		return;
 	}
 
+	// We can either split the mesh into multiple meshes - one for each poly group and create UVs for each with calls to UVAtlas (potentially in parallel)
+	// or we can update the adjacency and use a single call to UVAtlas.
+	//
+	// We might disable the Split because running our LayoutUV tool on the results it could mangle some triangles, while the single call
+	// to UVAtlas automatically does some layout.  This mangling happens some UVtriangles have opposite winding order from the rest of the group and they
+	// get put in their own chart..  I noticed that only when I remesh with "prevent normal flipping" I get these inverted uv triangles..
+	const bool bSplitIntoMultiple = true;
+
 	// The UV atlas callback uses a float progress-based interrupter. 
 	TFunction<bool(float)> Iterrupter = [Progress](float)->bool {return !Progress->Cancelled(); };
 
+	if (bSplitIntoMultiple && bRespectPolygroups && ResultMesh->HasTriangleGroups())
+	{
+		// Count the number of components for which we successufully generate UVs
+		int32 SuccessCount = 0;
+
+		// Predicate that reports true if two triangles are in the same group.
+		auto GroupPredicate = [this](int32 CurTri, int32 NbrTri)->bool
+		{
+			const int CurTriGroup = this->ResultMesh->GetTriangleGroup(CurTri);
+			const int NbrTriGroup = this->ResultMesh->GetTriangleGroup(NbrTri);
+			
+			return (CurTriGroup == NbrTriGroup);
+		};
+
+		FMeshConnectedComponents ConnectedComponents(ResultMesh.Get());
+		ConnectedComponents.FindConnectedTriangles(GroupPredicate); 
 
-	bool bSuccess = ComputeUVs(*ResultMesh,  Iterrupter);
+		// Create an array of meshes - one for each poly group
+		// NB TIndirectArray automatically deletes the meshes when it goes out of scope.
+		TIndirectArray<FDynamicMesh3> DynamicMeshArray;
 
+		for (int cidx = 0; cidx < ConnectedComponents.Num(); ++cidx)
+		{
+			// the tri ids for this component.
+			const TArray<int>& TrisIDArray = ConnectedComponents.GetComponent(cidx).Indices;
+
+			// create a new mesh from this component.
+			FDynamicMesh3* ComponentMesh = new FDynamicMesh3;
+			DynamicMeshArray.Add(ComponentMesh);
+
+			TMap<int32, int32> VtxIDMap; // Key: vert Id in source mesh -> Value: vert Id in component mesh.
+			for (auto TriID : TrisIDArray)
+			{
+
+				const FIndex3i RsltTriangle = ResultMesh->GetTriangle(TriID);
+				FIndex3i ComponentTriangle;
+				
+				// update the component mesh vertex buffer to accommodate this triangle
+				// also capture the VertIDs for this triangle wrt the component vertex buffer.
+				for (int i = 0; i < 3; ++i)
+				{
+					const int RsltVertID = RsltTriangle[i];
+					int ComponentVertID = -1;
+					if (const int* ComponetVertIDPtr = VtxIDMap.Find(RsltVertID))
+					{
+						ComponentVertID = *ComponetVertIDPtr;
+					}
+					else // Add this vert if we haven't already seen it.
+					{
+						const FVector3d Vertex = ResultMesh->GetVertex(RsltVertID);
+						ComponentVertID = ComponentMesh->AppendVertex(Vertex);
+						VtxIDMap.Add(RsltVertID, ComponentVertID);
+					}
+					
+					checkSlow(ComponentVertID != -1);
+
+					// update the component tri to this component vert.
+					ComponentTriangle[i] = ComponentVertID;
+				}
+
+				// add the triangle. 
+				ComponentMesh->AppendTriangle(ComponentTriangle);
+
+			}
+		}
+
+		
+
+		// Compute UVs for each component mesh.  NB: this could be parallelized over each
+		for (int32 MeshIdx = 0; MeshIdx < DynamicMeshArray.Num(); ++MeshIdx)
+		{
+			FDynamicMesh3& DynamicMesh = DynamicMeshArray[MeshIdx];
+			int32 NumTris = DynamicMesh.TriangleCount();
+			checkSlow(NumTris > 0);
+
+			if (NumTris > 0)
+			{
+				bool bComputedUVs = ComputeUVs(DynamicMesh, Iterrupter, false, AreaScaling);
+				
+				if (bComputedUVs) SuccessCount++;
+			}
+		}
+
+		// transfer the UVs back to the ResultMesh
+		if (SuccessCount > 0)
+		{
+
+			// Set up UV mesh overlay
+
+			const bool bHasAttributes = ResultMesh->HasAttributes();
+
+			if (bHasAttributes)
+			{
+				FDynamicMeshAttributeSet* Attributes = ResultMesh->Attributes();
+				Attributes->PrimaryUV()->ClearElements(); // delete existing UVs
+			}
+			else
+			{
+				// Add attrs for UVS
+				ResultMesh->EnableAttributes();
+			}
+
+			FDynamicMeshUVOverlay* UVOverlay = ResultMesh->Attributes()->PrimaryUV();
+
+			// loop over each component mesh, updating the result mesh UVs 
+
+			for (int cidx = 0; cidx < ConnectedComponents.Num(); ++cidx)
+			{
+
+				// map the tris in component mesh to source mesh
+				const TArray<int>& TrisIDArray = ConnectedComponents.GetComponent(cidx).Indices;
+
+				const FDynamicMesh3&  ComponentMesh = DynamicMeshArray[cidx];
+
+				// If the UV generation for this component failed, it will have no attributes
+				if (!ComponentMesh.HasAttributes())
+				{
+					continue;
+				}
+				// UV overlay for the component mesh
+				const FDynamicMeshUVOverlay* ComponentUVOverlay = ComponentMesh.Attributes()->PrimaryUV();
+
+				const int32 MaxComponentElementID = ComponentUVOverlay->MaxElementID();
+				
+				// Mapping from component mesh overlay element ID to result mesh overlay element ID
+				// initialized with -1.
+				TArray<int32> ElIDMap;
+				ElIDMap.Reserve(MaxComponentElementID + 1);
+				for (int32 i = 0; i < MaxComponentElementID + 1; ++i)
+				{
+					ElIDMap.Add(-1);
+				}
+
+				for (int UVTriID = 0; UVTriID < TrisIDArray.Num(); ++UVTriID)
+				{
+
+					// the ID for this triangle in the source mesh
+					const int RsltTriID = TrisIDArray[UVTriID];
+
+					// the vertex ids of the triangle in the source mesh.
+					const FIndex3i RsltTriangle = ResultMesh->GetTriangle(RsltTriID);
+
+					// the element IDs for the UVs on this triangle.
+					const FIndex3i ComponentTriElIDs = ComponentUVOverlay->GetTriangle(UVTriID);
+					FIndex3i TriElIDs; // the triangle we are constructing on the result mesh.
+
+					for (int i = 0; i < 3; ++i)
+					{
+
+						const int ComponentElID = ComponentTriElIDs[i];
+						
+						// the vert that we want this element associated with
+						const int VertID = RsltTriangle[i];
+						
+						if (ElIDMap[ComponentElID] == -1) // we haven't seen this element before
+						{
+							FVector2f UVvalue = ComponentUVOverlay->GetElement(ComponentElID);
+
+							// add this element to the result mesh
+							const int32 ElID = UVOverlay->AppendElement(UVvalue, VertID);
+
+							// record the element ID
+							ElIDMap[ComponentElID] = ElID;
+						}
+
+						TriElIDs[i] = ElIDMap[ComponentElID];
+					}
+
+					// Associate the UV tri with the source tri
+					UVOverlay->SetTriangle(RsltTriID, TriElIDs);
+
+				}
+
+			}
+		}
+
+	}
+	else
+	{
+		// Single call to UVAtlas - perhaps respecting the poly groups.
+		bool bSuccess = ComputeUVs(*ResultMesh, Iterrupter, bRespectPolygroups, AreaScaling);
+	}
 }
\ No newline at end of file
diff --git a/Engine/Plugins/Experimental/MeshModelingToolset/Source/ModelingOperatorsEditorOnly/Public/ParameterizationOps/ParameterizeMeshOp.h b/Engine/Plugins/Experimental/MeshModelingToolset/Source/ModelingOperatorsEditorOnly/Public/ParameterizationOps/ParameterizeMeshOp.h
--- a/Engine/Plugins/Experimental/MeshModelingToolset/Source/ModelingOperatorsEditorOnly/Public/ParameterizationOps/ParameterizeMeshOp.h
+++ b/Engine/Plugins/Experimental/MeshModelingToolset/Source/ModelingOperatorsEditorOnly/Public/ParameterizationOps/ParameterizeMeshOp.h
@@ -5,6 +5,7 @@
 #include "CoreMinimal.h"
 #include "Util/ProgressCancel.h"
 #include "ModelingOperators.h"
+#include "DynamicMeshAttributeSet.h"
 
 
 struct FMeshDescription;
@@ -26,11 +27,18 @@ public:
 	float Stretch;
 	int32 NumCharts;
 
+	// area scaling
+	bool bNormalizeAreas = true;
+	float AreaScaling = 1.0;
+
 	// Atlas Packing parameters
 	int32 Height = 512;
 	int32 Width = 512;
 	float Gutter = 2.5;
 
+	// Align UV boundaries with current poly group boundaries
+	bool bRespectPolygroups = true;
+
 	// set ability on protected transform.
 	void SetTransform(FTransform3d& XForm)
 	{
@@ -49,7 +57,7 @@ protected:
 	// dense index/vertex buffer based representation of the data needed for parameterization
 	struct FLinearMesh
 	{
-		FLinearMesh(const FDynamicMesh3& Mesh);
+		FLinearMesh(const FDynamicMesh3& Mesh, const bool bRespectPolygroups);
 
 		// Stripped down mesh
 		TArray<int32>   IndexBuffer;
@@ -64,5 +72,7 @@ protected:
 		
 	};
 
-	bool ComputeUVs(FDynamicMesh3& InOutMesh, TFunction<bool(float)>& Interrupter);
+	bool ComputeUVs(FDynamicMesh3& InOutMesh,  TFunction<bool(float)>& Interrupter, const bool bUsePolygroups = false, float GlobalScale = 1.0f);
+
+	void NormalizeUVAreas(const FDynamicMesh3* Mesh, FDynamicMeshUVOverlay* Overlay, float GlobalScale = 1.0f);
 };
diff --git a/Engine/Plugins/Experimental/ModelingToolsEditorMode/Source/ModelingToolsEditorMode/ModelingToolsEditorMode.Build.cs b/Engine/Plugins/Experimental/ModelingToolsEditorMode/Source/ModelingToolsEditorMode/ModelingToolsEditorMode.Build.cs
--- a/Engine/Plugins/Experimental/ModelingToolsEditorMode/Source/ModelingToolsEditorMode/ModelingToolsEditorMode.Build.cs
+++ b/Engine/Plugins/Experimental/ModelingToolsEditorMode/Source/ModelingToolsEditorMode/ModelingToolsEditorMode.Build.cs
@@ -52,7 +52,9 @@ public class ModelingToolsEditorMode : ModuleRules
 				"MeshModelingTools",
 				"MeshModelingToolsEditorOnly",
 
-                "ViewportInteraction"
+                "ViewportInteraction",
+
+				"StylusInput"
 
 				// ... add private dependencies that you statically link with here ...	
 			}
diff --git a/Engine/Plugins/Experimental/ModelingToolsEditorMode/Source/ModelingToolsEditorMode/Private/ModelingToolsEditorModeStyle.cpp b/Engine/Plugins/Experimental/ModelingToolsEditorMode/Source/ModelingToolsEditorMode/Private/ModelingToolsEditorModeStyle.cpp
--- a/Engine/Plugins/Experimental/ModelingToolsEditorMode/Source/ModelingToolsEditorMode/Private/ModelingToolsEditorModeStyle.cpp
+++ b/Engine/Plugins/Experimental/ModelingToolsEditorMode/Source/ModelingToolsEditorMode/Private/ModelingToolsEditorModeStyle.cpp
@@ -71,7 +71,7 @@ void FModelingToolsEditorModeStyle::Initialize()
 		// Accept/Cancel/Complete active tool
 
 		StyleSet->Set("LevelEditor.ModelingToolsMode", new IMAGE_PLUGIN_BRUSH("Icons/icon_ModelingToolsEditorMode", FVector2D(40.0f, 40.0f)));
-		StyleSet->Set("LevelEditor.ModelingToolsMode.Small", new IMAGE_PLUGIN_BRUSH("Icons/icon_ModelingToolsEditorMode", FVector2D(40.0f, 40.0f)));
+		StyleSet->Set("LevelEditor.ModelingToolsMode.Small", new IMAGE_PLUGIN_BRUSH("Icons/icon_ModelingToolsEditorMode", FVector2D(20.0f, 20.0f)));
 
 		// NOTE:  Old-style, need to be replaced: 
 		StyleSet->Set("ModelingToolsManagerCommands.CancelActiveTool", new IMAGE_PLUGIN_BRUSH("Icons/icon_ActiveTool_Cancel_40x", Icon20x20));
@@ -106,6 +106,8 @@ void FModelingToolsEditorModeStyle::Initialize()
 		StyleSet->Set("ModelingToolsManagerCommands.BeginSculptMeshTool.Small", 		new IMAGE_PLUGIN_BRUSH("Icons/Sculpt_40x", 			Icon20x20));
 		StyleSet->Set("ModelingToolsManagerCommands.BeginPolyEditTool", 				new IMAGE_PLUGIN_BRUSH("Icons/PolyEdit_40x", 		Icon20x20));
 		StyleSet->Set("ModelingToolsManagerCommands.BeginPolyEditTool.Small", 			new IMAGE_PLUGIN_BRUSH("Icons/PolyEdit_40x", 		Icon20x20));
+		StyleSet->Set("ModelingToolsManagerCommands.BeginPolyDeformTool", 				new IMAGE_PLUGIN_BRUSH("Icons/PolyEdit_40x", 		Icon20x20));
+		StyleSet->Set("ModelingToolsManagerCommands.BeginPolyDeformTool.Small", 		new IMAGE_PLUGIN_BRUSH("Icons/PolyEdit_40x", 		Icon20x20));
 		StyleSet->Set("ModelingToolsManagerCommands.BeginDisplaceMeshTool", 			new IMAGE_PLUGIN_BRUSH("Icons/Displace_40x", 		Icon20x20));
 		StyleSet->Set("ModelingToolsManagerCommands.BeginDisplaceMeshTool.Small", 		new IMAGE_PLUGIN_BRUSH("Icons/Displace_40x", 		Icon20x20));
 		StyleSet->Set("ModelingToolsManagerCommands.BeginTransformMeshesTool", 			new IMAGE_PLUGIN_BRUSH("Icons/Transform_40x", 		Icon20x20));
@@ -119,8 +121,12 @@ void FModelingToolsEditorModeStyle::Initialize()
 		StyleSet->Set("ModelingToolsManagerCommands.BeginSimplifyMeshTool.Small", 		new IMAGE_PLUGIN_BRUSH("Icons/Simplify_40x", 		Icon20x20));
 		StyleSet->Set("ModelingToolsManagerCommands.BeginEditNormalsTool", 				new IMAGE_PLUGIN_BRUSH("Icons/Normals_40x",			Icon20x20));
 		StyleSet->Set("ModelingToolsManagerCommands.BeginEditNormalsTool.Small", 		new IMAGE_PLUGIN_BRUSH("Icons/Normals_40x",			Icon20x20));
+		StyleSet->Set("ModelingToolsManagerCommands.BeginRemoveOccludedTrianglesTool", 				new IMAGE_PLUGIN_BRUSH("Icons/Jacket_40x",			Icon20x20));
+		StyleSet->Set("ModelingToolsManagerCommands.BeginRemoveOccludedTrianglesTool.Small", 		new IMAGE_PLUGIN_BRUSH("Icons/Jacket_40x",			Icon20x20));
 		StyleSet->Set("ModelingToolsManagerCommands.BeginUVProjectionTool", 			new IMAGE_PLUGIN_BRUSH("Icons/UVProjection_40x", 	Icon20x20));
 		StyleSet->Set("ModelingToolsManagerCommands.BeginUVProjectionTool.Small", 		new IMAGE_PLUGIN_BRUSH("Icons/UVProjection_40x", 	Icon20x20));
+		StyleSet->Set("ModelingToolsManagerCommands.BeginUVLayoutTool", 			new IMAGE_PLUGIN_BRUSH("Icons/UVLayout_40x", 	Icon20x20));
+		StyleSet->Set("ModelingToolsManagerCommands.BeginUVLayoutTool.Small", 		new IMAGE_PLUGIN_BRUSH("Icons/UVLayout_40x", 	Icon20x20));
 		StyleSet->Set("ModelingToolsManagerCommands.BeginVoxelMergeTool", 				new IMAGE_PLUGIN_BRUSH("Icons/VoxMerge_40x", 		Icon20x20));
 		StyleSet->Set("ModelingToolsManagerCommands.BeginVoxelMergeTool.Small", 		new IMAGE_PLUGIN_BRUSH("Icons/VoxMerge_40x", 		Icon20x20));
 		StyleSet->Set("ModelingToolsManagerCommands.BeginVoxelBooleanTool", 			new IMAGE_PLUGIN_BRUSH("Icons/VoxBoolean_40x", 		Icon20x20));
diff --git a/Engine/Plugins/Experimental/ModelingToolsEditorMode/Source/ModelingToolsEditorMode/Private/ModelingToolsManagerActions.cpp b/Engine/Plugins/Experimental/ModelingToolsEditorMode/Source/ModelingToolsEditorMode/Private/ModelingToolsManagerActions.cpp
--- a/Engine/Plugins/Experimental/ModelingToolsEditorMode/Source/ModelingToolsEditorMode/Private/ModelingToolsManagerActions.cpp
+++ b/Engine/Plugins/Experimental/ModelingToolsEditorMode/Source/ModelingToolsEditorMode/Private/ModelingToolsManagerActions.cpp
@@ -22,21 +22,27 @@ FModelingToolsManagerCommands::FModelingToolsManagerCommands() :
 void FModelingToolsManagerCommands::RegisterCommands()
 {
 	UI_COMMAND(BeginAddPrimitiveTool, "Primitive", "Start the Add Primitive Tool", EUserInterfaceActionType::Button, FInputChord());
+	UI_COMMAND(BeginAddPatchTool, "Patch", "Start the Add Patch Tool", EUserInterfaceActionType::Button, FInputChord());
 	UI_COMMAND(BeginDrawPolygonTool, "Polygon", "Start the Draw Polygon Tool", EUserInterfaceActionType::Button, FInputChord());
 	UI_COMMAND(BeginShapeSprayTool, "ShapeSpray", "Start the Shape Spray Tool", EUserInterfaceActionType::Button, FInputChord());
 
 	UI_COMMAND(BeginSculptMeshTool, "Sculpt", "Start the Sculpt Mesh Tool", EUserInterfaceActionType::Button, FInputChord());
 	UI_COMMAND(BeginPolyEditTool, "PolyEdit", "Start the PolyEdit Tool", EUserInterfaceActionType::Button, FInputChord());
+	UI_COMMAND(BeginPolyDeformTool, "PolyDeform", "Start the PolyDeform Tool", EUserInterfaceActionType::Button, FInputChord());
 	UI_COMMAND(BeginSmoothMeshTool, "Smooth", "Start the Smooth Mesh Tool", EUserInterfaceActionType::Button, FInputChord());
 	UI_COMMAND(BeginDisplaceMeshTool, "Displace", "Start the Displace Mesh Tool", EUserInterfaceActionType::Button, FInputChord());
 	UI_COMMAND(BeginMeshSpaceDeformerTool, "SpaceDeform", "Start the Mesh Space Deformer Tool", EUserInterfaceActionType::Button, FInputChord());
 	UI_COMMAND(BeginTransformMeshesTool, "Transform", "Start the Transform Tool", EUserInterfaceActionType::Button, FInputChord());
+	UI_COMMAND(BeginEditPivotTool, "Edit Pivot", "Start the Edit Pivot Tool", EUserInterfaceActionType::Button, FInputChord());
+	UI_COMMAND(BeginBakeTransformTool, "Bake XForm", "Start the Bake Transform Tool", EUserInterfaceActionType::Button, FInputChord());
 
 	UI_COMMAND(BeginRemeshSculptMeshTool, "DynaSculpt", "Start the Liquid Sculpt Mesh Tool", EUserInterfaceActionType::Button, FInputChord());
 	UI_COMMAND(BeginRemeshMeshTool, "Remesh", "Start the Remesh Tool", EUserInterfaceActionType::Button, FInputChord());
 	UI_COMMAND(BeginSimplifyMeshTool, "Simplify", "Start the Simplify Mesh Tool", EUserInterfaceActionType::Button, FInputChord());
 	UI_COMMAND(BeginEditNormalsTool, "Normals", "Start the Edit Normals Tool", EUserInterfaceActionType::Button, FInputChord());
+	UI_COMMAND(BeginRemoveOccludedTrianglesTool, "Jacketing", "Start the Jacketing Tool", EUserInterfaceActionType::Button, FInputChord());
 	UI_COMMAND(BeginUVProjectionTool, "UVProjection", "Start the UV Projection Tool", EUserInterfaceActionType::Button, FInputChord());
+	UI_COMMAND(BeginUVLayoutTool, "UVLayout", "Start the UV Layout Tool", EUserInterfaceActionType::Button, FInputChord());
 	UI_COMMAND(BeginVoxelMergeTool, "VoxMerge", "Start the Voxel Merge Tool", EUserInterfaceActionType::Button, FInputChord());
 	UI_COMMAND(BeginVoxelBooleanTool, "VoxBoolean", "Start the Voxel Boolean Tool", EUserInterfaceActionType::Button, FInputChord());
 	UI_COMMAND(BeginPlaneCutTool, "PlaneCut", "Start the Plane Cut Tool", EUserInterfaceActionType::Button, FInputChord());
@@ -47,7 +53,8 @@ void FModelingToolsManagerCommands::RegisterCommands()
 	UI_COMMAND(BeginMeshInspectorTool, "Inspector", "Start the Mesh Inspector Tool", EUserInterfaceActionType::Button, FInputChord());
 	UI_COMMAND(BeginWeldEdgesTool, "Weld Edges", "Start the Weld Edges Tool", EUserInterfaceActionType::Button, FInputChord());
 	UI_COMMAND(BeginPolyGroupsTool, "PolyGroups", "Start the Poly Groups Tool", EUserInterfaceActionType::Button, FInputChord());
-	UI_COMMAND(BeginAttributeEditorTool, "AttributeEditor", "Start the Attribute Editor Tool", EUserInterfaceActionType::Button, FInputChord());
+	UI_COMMAND(BeginEditMeshMaterialsTool, "Edit Mats", "Start the Material Editing Tool", EUserInterfaceActionType::Button, FInputChord());
+	UI_COMMAND(BeginAttributeEditorTool, "Edit Attribs", "Start the Attribute Editor Tool", EUserInterfaceActionType::Button, FInputChord());
 
 	UI_COMMAND(AcceptActiveTool, "Accept", "Accept the active tool", EUserInterfaceActionType::Button, FInputChord());
 	UI_COMMAND(CancelActiveTool, "Cancel", "Cancel the active tool", EUserInterfaceActionType::Button, FInputChord());
diff --git a/Engine/Plugins/Experimental/PlanarCutPlugin/Source/PlanarCut/Private/PlanarCut.cpp b/Engine/Plugins/Experimental/PlanarCutPlugin/Source/PlanarCut/Private/PlanarCut.cpp
--- a/Engine/Plugins/Experimental/PlanarCutPlugin/Source/PlanarCut/Private/PlanarCut.cpp
+++ b/Engine/Plugins/Experimental/PlanarCutPlugin/Source/PlanarCut/Private/PlanarCut.cpp
@@ -1246,8 +1246,8 @@ void CutWithPlanarCellsHelper(
 				return false;
 			}
 			int sides[3] = { 0,0,0 };
-			sides[PlaneSide(Plane, Box.Min) + 1]++;
-			sides[PlaneSide(Plane, Box.Max) + 1]++;
+			sides[PlaneSide(Plane, (FVector)Box.Min) + 1]++;
+			sides[PlaneSide(Plane, (FVector)Box.Max) + 1]++;
 			sides[PlaneSide(Plane, FVector(Box.Max.X, Box.Min.Y, Box.Min.Z)) + 1]++;
 			sides[PlaneSide(Plane, FVector(Box.Min.X, Box.Max.Y, Box.Min.Z)) + 1]++;
 			sides[PlaneSide(Plane, FVector(Box.Max.X, Box.Max.Y, Box.Min.Z)) + 1]++;
@@ -1575,7 +1575,7 @@ void CutWithPlanarCellsHelper(
 		PlaneTriangulationInfo& Triangulation = PlaneTriangulations[PlaneIdx];
 		int32 NumBoundary = BoundaryIndices.Num();
 		FVector PlaneNormal(Plane.X, Plane.Y, Plane.Z);
-		FVector Origin = PlaneFrames[PlaneIdx].Origin;
+		FVector Origin = (FVector)PlaneFrames[PlaneIdx].Origin;
 
 		// check if constrained Delaunay triangulation problem needed for plane (false if no geometry was touching the planar facet)
 		bool bAnyElementsOnPlane = EdgesOnPlane[PlaneIdx].Num() + TrianglesOnPlane[PlaneIdx].Num() > 0;
@@ -1884,7 +1884,7 @@ void CutWithPlanarCellsHelper(
 			ensure(Arrangement.Graph.IsCompact());
 			for (int32 VertIdx = 0; VertIdx < Arrangement.Graph.MaxVertexID(); VertIdx++)
 			{
-				Triangulation.LocalVertices.Add(PlaneFrames[PlaneIdx].UnProject(Arrangement.Graph.GetVertex(VertIdx)));
+				Triangulation.LocalVertices.Add((FVector)PlaneFrames[PlaneIdx].UnProject(Arrangement.Graph.GetVertex(VertIdx)));
 			}
 			
 			for (FIntVector& Face : PlaneTriangulation)
@@ -1915,14 +1915,14 @@ void CutWithPlanarCellsHelper(
 				FVector3d Z = PlaneNormal * Amplitude;
 				for (int32 VertexIdx : NoiseVertexIndices)
 				{
-					FVector2D V = Arrangement.Graph.GetVertex(VertexIdx) * Frequency;
+					FVector2D V = (FVector2D)(Arrangement.Graph.GetVertex(VertexIdx) * Frequency);
 					float NoiseValue = 0;
 					float OctaveScale = 1;
 					for (int32 Octave = 0; Octave < Octaves; Octave++, OctaveScale *= 2)
 					{
 						NoiseValue += FMath::PerlinNoise2D(V * OctaveScale * AverageGlobalScale) / OctaveScale;
 					}
-					Triangulation.LocalVertices[VertexIdx] += Z * NoiseValue * AverageGlobalScaleInv;
+					Triangulation.LocalVertices[VertexIdx] += (FVector)(Z * NoiseValue * AverageGlobalScaleInv);
 				}
 			}
 		}
@@ -1964,8 +1964,8 @@ void CutWithPlanarCellsHelper(
 		if (NumLocalVertices > 0)
 		{
 			Triangulation.LocalUVs.SetNum(NumLocalVertices);
-			FVector FrameX = PlaneFrames[PlaneIdx].X;
-			FVector FrameY = PlaneFrames[PlaneIdx].Y;
+			FVector FrameX = (FVector)PlaneFrames[PlaneIdx].X;
+			FVector FrameY = (FVector)PlaneFrames[PlaneIdx].Y;
 			FVector LocalOrigin = Triangulation.LocalVertices[0];
 			float MinX = FMathf::MaxReal;
 			float MinY = FMathf::MaxReal;
@@ -2064,10 +2064,10 @@ void CutWithPlanarCellsHelper(
 			AddedVerticesCollection.Normal[AddIdx] = PlaneNormal;
 			AddedVerticesCollection.Normal[AddIdx + NumLocalVertices] = -AddedVerticesCollection.Normal[AddIdx];
 
-			AddedVerticesCollection.TangentU[AddIdx] = PlaneFrames[PlaneIdx].X;
+			AddedVerticesCollection.TangentU[AddIdx] = (FVector)PlaneFrames[PlaneIdx].X;
 			AddedVerticesCollection.TangentU[AddIdx + NumLocalVertices] = -AddedVerticesCollection.TangentU[AddIdx];
 
-			AddedVerticesCollection.TangentV[AddIdx] = PlaneFrames[PlaneIdx].Y;
+			AddedVerticesCollection.TangentV[AddIdx] = (FVector)PlaneFrames[PlaneIdx].Y;
 			AddedVerticesCollection.TangentV[AddIdx + NumLocalVertices] = AddedVerticesCollection.TangentV[AddIdx];
 
 			// TODO: also set color?
@@ -2114,7 +2114,7 @@ void TransformPlanes(const FTransform& Transform, const FPlanarCells& Ref, TArra
 		ScaleVec /= ScaleMaxAbs;
 	}
 	FVector3d NormalScale(ScaleVec.Y*ScaleVec.Z*ScaleDetSign, ScaleVec.X*ScaleVec.Z*ScaleDetSign, ScaleVec.X*ScaleVec.Y*ScaleDetSign);
-	NormalTransform.SetScale3D(NormalScale);
+	NormalTransform.SetScale3D((FVector)NormalScale);
 	
 	Planes.SetNum(Ref.Planes.Num());
 	for (int32 PlaneIdx = 0, PlanesNum = Planes.Num(); PlaneIdx < PlanesNum; PlaneIdx++)
@@ -2475,7 +2475,7 @@ int32 CutMultipleWithMultiplePlanes(
 									FTransform NbrToLocalTransform = GlobalTransforms[NbrTransformIdx] * GlobalTransforms[ChildTransformIdx].Inverse();
 									auto NbrToLocal = [&](const FVector3d& V)
 									{
-										return FVector3d(NbrToLocalTransform.TransformPosition(V));
+										return FVector3d(NbrToLocalTransform.TransformPosition((FVector)V));
 									};
 									double OutDist;
 									ChildTree->FindNearestTriangles(*NbrTree, NbrToLocal, OutDist, ProximityThresholdDist);
@@ -2513,7 +2513,7 @@ int32 CutMultipleWithMultiplePlanes(
 								FTransform NbrToLocalTransform = GlobalTransforms[NbrTransformIdx] * GlobalTransforms[ChildTransformIdx].Inverse();
 								auto NbrToLocal = [&](const FVector3d& V)
 								{
-									return FVector3d(NbrToLocalTransform.TransformPosition(V));
+									return FVector3d(NbrToLocalTransform.TransformPosition((FVector)V));
 								};
 								double OutDist;
 								ChildTree->FindNearestTriangles(*NbrTree, NbrToLocal, OutDist, ProximityThresholdDist);
diff --git a/Engine/Plugins/Experimental/ProxyLODPlugin/Source/ProxyLOD/Private/ProxyLODMeshParameterization.cpp b/Engine/Plugins/Experimental/ProxyLODPlugin/Source/ProxyLOD/Private/ProxyLODMeshParameterization.cpp
--- a/Engine/Plugins/Experimental/ProxyLODPlugin/Source/ProxyLOD/Private/ProxyLODMeshParameterization.cpp
+++ b/Engine/Plugins/Experimental/ProxyLODPlugin/Source/ProxyLOD/Private/ProxyLODMeshParameterization.cpp
@@ -55,7 +55,8 @@ bool ProxyLOD::GenerateUVs(const FTextureAtlasDesc& TextureAtlasDesc,
 
 	// Verify the mesh is valid
 	{
-		HRESULT ValidateHR = DirectX::Validate(Indices.GetData(), NumFaces, NumVerts, AdjacencyArray.GetData(), DirectX::VALIDATE_BOWTIES, NULL);
+	
+		HRESULT ValidateHR = DirectX::Validate(Indices.GetData(), NumFaces, NumVerts, AdjacencyArray.GetData(), DirectX::VALIDATE_DEFAULT, NULL);
 		if (FAILED(ValidateHR))
 		{
 			return false;
@@ -265,6 +266,8 @@ bool ProxyLOD::GenerateUVs(FVertexDataMesh& InOutMesh, const FTextureAtlasDesc&
 		DirectX::UVATLAS_DEFAULT, vb, ib,
 		&facePartitioning, &vertexRemapArray, &maxStretchUsed, &numChartsUsed);
 
+	delete[] pIMTArray;
+
 	if (MaxStretchOut)
 	{
 		*MaxStretchOut = maxStretchUsed;
@@ -275,7 +278,7 @@ bool ProxyLOD::GenerateUVs(FVertexDataMesh& InOutMesh, const FTextureAtlasDesc&
 	}
 	if (FAILED(hr)) return false;
 
-	if (pIMTArray) delete[] pIMTArray;
+	
 	
 
 	// testing
diff --git a/Engine/Plugins/Experimental/ProxyLODPlugin/Source/ProxyLOD/Public/ProxyLODVolume.h b/Engine/Plugins/Experimental/ProxyLODPlugin/Source/ProxyLOD/Public/ProxyLODVolume.h
--- a/Engine/Plugins/Experimental/ProxyLODPlugin/Source/ProxyLOD/Public/ProxyLODVolume.h
+++ b/Engine/Plugins/Experimental/ProxyLODPlugin/Source/ProxyLOD/Public/ProxyLODVolume.h
@@ -74,6 +74,10 @@ public :
 			: Mesh(nullptr)
 		{}
 
+		explicit FPlacedMesh(const FMeshDescription* MeshIn, const FTransform& TransformIn = FTransform::Identity)
+			: Mesh(MeshIn), Transform(TransformIn)
+		{}
+
 		FPlacedMesh(const FPlacedMesh& other)
 			: Mesh(other.Mesh)
 			, Transform(other.Transform)
diff --git a/Engine/Plugins/Experimental/SkeletalReduction/Source/Private/SkeletalMeshReductionPlugin.cpp b/Engine/Plugins/Experimental/SkeletalReduction/Source/Private/SkeletalMeshReductionPlugin.cpp
--- a/Engine/Plugins/Experimental/SkeletalReduction/Source/Private/SkeletalMeshReductionPlugin.cpp
+++ b/Engine/Plugins/Experimental/SkeletalReduction/Source/Private/SkeletalMeshReductionPlugin.cpp
@@ -1987,12 +1987,12 @@ void FQuadricSkeletalMeshReduction::ReduceSkeletalMesh(USkeletalMesh& SkeletalMe
 
 		if (!bReducingSourceModel && Old)
 		{
+			bool bIsOldRawSkelMeshEmpty = SkeletalMesh.IsLODImportedDataEmpty(LODIndex);
 			//We need to backup the original RawSkeletalMeshBulkData in case it was an imported LOD
-			if (!bLODModelAdded && !Old->RawSkeletalMeshBulkData.IsEmpty())
+			if (!bLODModelAdded && !bIsOldRawSkelMeshEmpty)
 			{
-				Old->RawSkeletalMeshBulkData.LoadRawMesh(RawMesh);
-				GeoImportVersion = Old->RawSkeletalMeshBulkData.GeoImportVersion;
-				SkinningImportVersion = Old->RawSkeletalMeshBulkData.SkinningImportVersion;
+				SkeletalMesh.LoadLODImportedData(LODIndex, RawMesh);
+				SkeletalMesh.GetLODImportedDataVersions(LODIndex, GeoImportVersion, SkinningImportVersion);
 			}
 			//If the delegate is not bound 
 			if (!Settings.OnDeleteLODModelDelegate.IsBound())
@@ -2009,10 +2009,8 @@ void FQuadricSkeletalMeshReduction::ReduceSkeletalMesh(USkeletalMesh& SkeletalMe
 		}
 		else if(bReducingSourceModel)
 		{
-			//In case we reduce the source model we want to keep the original import data
-			SrcModel->RawSkeletalMeshBulkData.LoadRawMesh(RawMesh);
-			GeoImportVersion = SrcModel->RawSkeletalMeshBulkData.GeoImportVersion;
-			SkinningImportVersion = SrcModel->RawSkeletalMeshBulkData.SkinningImportVersion;
+			SkeletalMesh.LoadLODImportedData(BaseLOD, RawMesh);
+			SkeletalMesh.GetLODImportedDataVersions(BaseLOD, GeoImportVersion, SkinningImportVersion);
 		}
 	}
 
@@ -2146,9 +2144,8 @@ void FQuadricSkeletalMeshReduction::ReduceSkeletalMesh(USkeletalMesh& SkeletalMe
 	if ((bReducingSourceModel || !bLODModelAdded ) && RawMesh.Points.Num() > 0)
 	{
 		//Put back the original import data, we need it to allow inline reduction and skeletal mesh split workflow
-		SkeletalMeshResource.LODModels[LODIndex].RawSkeletalMeshBulkData.SaveRawMesh(RawMesh);
-		SkeletalMeshResource.LODModels[LODIndex].RawSkeletalMeshBulkData.GeoImportVersion = GeoImportVersion;
-		SkeletalMeshResource.LODModels[LODIndex].RawSkeletalMeshBulkData.SkinningImportVersion = SkinningImportVersion;
+		SkeletalMesh.SaveLODImportedData(LODIndex, RawMesh);
+		SkeletalMesh.SetLODImportedDataVersions(LODIndex, GeoImportVersion, SkinningImportVersion);
 	}
 
 	SkeletalMesh.CalculateRequiredBones(SkeletalMeshResource.LODModels[LODIndex], SkeletalMesh.RefSkeleton, &BonesToRemove);
diff --git a/Engine/Plugins/Lumin/MagicLeap/Source/MagicLeap/Private/MagicLeapHMD.cpp b/Engine/Plugins/Lumin/MagicLeap/Source/MagicLeap/Private/MagicLeapHMD.cpp
--- a/Engine/Plugins/Lumin/MagicLeap/Source/MagicLeap/Private/MagicLeapHMD.cpp
+++ b/Engine/Plugins/Lumin/MagicLeap/Source/MagicLeap/Private/MagicLeapHMD.cpp
@@ -667,7 +667,7 @@ bool FMagicLeapHMD::IsHeadTrackingAllowed() const
 	if (GIsEditor)
 	{
 		UEditorEngine* EdEngine = Cast<UEditorEngine>(GEngine);
-		return ((!EdEngine || EdEngine->bUseVRPreviewForPlayWorld || GetDefault<ULevelEditorPlaySettings>()->ViewportGetsHMDControl) && GEngine->IsStereoscopic3D());
+		return ((!EdEngine || EdEngine->IsVRPreviewActive() || GetDefault<ULevelEditorPlaySettings>()->ViewportGetsHMDControl) && GEngine->IsStereoscopic3D());
 	}
 #endif // WITH_EDITOR
 	return (GEngine->IsStereoscopic3D());
diff --git a/Engine/Plugins/MeshPainting/Source/MeshPaintEditorMode/Private/BrushSettingsCustomization.cpp b/Engine/Plugins/MeshPainting/Source/MeshPaintEditorMode/Private/BrushSettingsCustomization.cpp
new file mode 100644
--- /dev/null
+++ b/Engine/Plugins/MeshPainting/Source/MeshPaintEditorMode/Private/BrushSettingsCustomization.cpp
@@ -0,0 +1,374 @@
+// Copyright 1998-2019 Epic Games, Inc. All Rights Reserved.
+
+#include "BrushSettingsCustomization.h"
+
+#include "DetailLayoutBuilder.h"
+#include "DetailCategoryBuilder.h"
+#include "IDetailPropertyRow.h"
+#include "IDetailChildrenBuilder.h"
+#include "DetailWidgetRow.h"
+#include "Widgets/Text/STextBlock.h"
+#include "PropertyRestriction.h"
+#include "Engine/Texture2D.h"
+#include "Widgets/SBoxPanel.h"
+#include "Widgets/Input/SCheckBox.h"
+#include "Widgets/Input/SNumericEntryBox.h"
+#include "PropertyCustomizationHelpers.h"
+#include "Widgets/Input/SButton.h"
+#include "Widgets/Layout/SBox.h"
+#include "ScopedTransaction.h"
+#include "MeshPaintHelpers.h"
+#include "MeshVertexPaintingTool.h"
+#include "MeshPaintHelpers.h"
+#include "InteractiveToolManager.h"
+#include "Editor.h"
+#include "EditorModeManager.h"
+#include "MeshPaintMode.h"
+#include "Components/SkeletalMeshComponent.h"
+
+#define LOCTEXT_NAMESPACE "MeshPaintCustomization"
+
+TSharedRef<SHorizontalBox> CreateColorChannelWidget(TSharedRef<IPropertyHandle> ChannelProperty)
+{
+	return SNew(SHorizontalBox)
+		+ SHorizontalBox::Slot()
+		.AutoWidth()
+		[
+			ChannelProperty->CreatePropertyValueWidget()
+		]
+	+ SHorizontalBox::Slot()
+		.AutoWidth()
+		.Padding(4.0f, 0.0f, 0.0f, 0.0f)
+		[
+			ChannelProperty->CreatePropertyNameWidget()
+		];
+}
+
+
+TSharedRef<IDetailCustomization> FVertexPaintingSettingsCustomization::MakeInstance()
+{
+	return MakeShareable(new FVertexPaintingSettingsCustomization);
+}
+
+void FVertexPaintingSettingsCustomization::CustomizeDetails(IDetailLayoutBuilder& DetailLayout)
+{	
+	IDetailCategoryBuilder& VertexCategory = DetailLayout.EditCategory(TEXT("VertexPainting"));
+
+	VertexCategory.AddCustomRow(NSLOCTEXT("VertexPaintSettings", "InstanceColorSize", "Instance Color Size"))
+		.WholeRowContent()
+		[
+			SNew(STextBlock)
+			.Text_Lambda([]() -> FText { return FText::Format(FTextFormat::FromString(TEXT("Instance Color Size: {0} KB")), UMeshPaintMode::GetMeshPaintMode()->GetCachedVertexDataSize() / 1024.f); })
+		];
+
+	/** Creates a custom widget row containing all color channel flags */	
+	TSharedRef<IPropertyHandle> PaintColor = DetailLayout.GetProperty("PaintColor", UMeshVertexPaintingToolProperties::StaticClass());
+	PaintColor->MarkHiddenByCustomization();
+	TSharedRef<IPropertyHandle> EraseColor = DetailLayout.GetProperty("EraseColor", UMeshVertexPaintingToolProperties::StaticClass());
+	EraseColor->MarkHiddenByCustomization();
+
+	// Customize paint color with a swap button
+	{
+		TSharedPtr<SWidget> NameWidget;
+		TSharedPtr<SWidget> ValueWidget;
+		
+		IDetailPropertyRow& PaintColorProp = VertexCategory.AddProperty(PaintColor);
+		PaintColorProp.GetDefaultWidgets(NameWidget, ValueWidget, false);
+		FDetailWidgetRow& Row = PaintColorProp.CustomWidget(true);
+		Row.NameContent()
+		[
+			NameWidget.ToSharedRef()
+		];
+
+		Row.ValueContent()
+		.MinDesiredWidth(250)
+		.MaxDesiredWidth(0)
+		[
+			SNew(SHorizontalBox)
+			+ SHorizontalBox::Slot()
+			.FillWidth(1.0)
+			.HAlign(HAlign_Left)
+			[
+				SNew(SBox)
+				.WidthOverride(250.f)
+				[
+					ValueWidget.ToSharedRef()
+				]
+			]
+			+ SHorizontalBox::Slot()
+			.HAlign(HAlign_Center)
+			.AutoWidth()
+			[
+				SNew(SButton)
+				.ButtonStyle(FEditorStyle::Get(), "HoverHintOnly")
+				.ToolTipText(NSLOCTEXT("VertexPaintSettings", "SwapColors", "Swap Paint and Erase Colors"))
+				.HAlign(HAlign_Center)
+				.VAlign(VAlign_Center)
+				.OnClicked(this, &FVertexPaintingSettingsCustomization::OnSwapColorsClicked, PaintColor, EraseColor)
+				.ContentPadding(0)
+				[
+					SNew(SImage).Image(FEditorStyle::GetBrush("MeshPaint.Swap"))
+				]
+			]
+		];
+	}
+
+	{
+		IDetailPropertyRow& EraseColorProp = VertexCategory.AddProperty(EraseColor);
+
+		TSharedPtr<SWidget> NameWidget;
+		TSharedPtr<SWidget> ValueWidget;
+		
+		FDetailWidgetRow& Row = EraseColorProp.CustomWidget(true);
+		Row.ValueContent().MinDesiredWidth(250 - 16.f);
+		EraseColorProp.GetDefaultWidgets(NameWidget, ValueWidget, Row, false);
+	}
+
+}
+
+void FVertexPaintingSettingsCustomization::OnTextureWeightTypeChanged(TSharedRef<IPropertyHandle> WeightTypeProperty, TSharedRef<IPropertyHandle> PaintWeightProperty, TSharedRef<IPropertyHandle> EraseWeightProperty)
+{
+	UEnum* ImportTypeEnum = StaticEnum<ETexturePaintWeightIndex>();
+	uint8 EnumValue = 0;
+	WeightTypeProperty->GetValue(EnumValue);	
+
+	BlendPaintEnumRestriction->RemoveAll();
+	for (uint8 EnumIndex = 0; EnumIndex < (ImportTypeEnum->GetMaxEnumValue() + 1); ++EnumIndex)
+	{
+		if ((EnumIndex + 1) > EnumValue)
+		{
+			FString EnumName = ImportTypeEnum->GetNameByValue(EnumIndex).ToString();
+			EnumName.RemoveFromStart("ETexturePaintIndex::");
+			BlendPaintEnumRestriction->AddDisabledValue(EnumName);
+		}
+	}
+
+	uint8 Value = 0;
+	PaintWeightProperty->GetValue(Value);
+	Value = FMath::Clamp<uint8>(Value, 0, EnumValue - 1);
+	PaintWeightProperty->SetValue(Value);
+
+	EraseWeightProperty->GetValue(Value);
+	Value = FMath::Clamp<uint8>(Value, 0, EnumValue - 1);
+	EraseWeightProperty->SetValue(Value);
+}
+
+FReply FVertexPaintingSettingsCustomization::OnSwapColorsClicked(TSharedRef<IPropertyHandle> PaintColor, TSharedRef<IPropertyHandle> EraseColor)
+{
+	FScopedTransaction Transaction(NSLOCTEXT("VertexPaintSettings", "SwapColorsTransation", "Swap paint and erase colors"));
+
+	UMeshVertexPaintingToolProperties* Settings = UMeshPaintMode::GetVertexToolProperties();
+	if (Settings)
+	{
+		Settings->Modify();
+		
+		FLinearColor TempPaintColor = Settings->PaintColor;
+		Settings->PaintColor = Settings->EraseColor;
+		Settings->EraseColor = TempPaintColor;
+		PaintColor->NotifyPostChange(EPropertyChangeType::ValueSet);
+		EraseColor->NotifyPostChange(EPropertyChangeType::ValueSet);
+	}
+	return FReply::Handled();
+}
+
+
+
+TSharedRef<IDetailCustomization> FColorPaintingSettingsCustomization::MakeInstance()
+{
+	return MakeShareable(new FColorPaintingSettingsCustomization);
+}
+
+void FColorPaintingSettingsCustomization::CustomizeDetails(IDetailLayoutBuilder& DetailLayout)
+{
+	FVertexPaintingSettingsCustomization::CustomizeDetails(DetailLayout);
+
+	IDetailCategoryBuilder& ColorCategory = DetailLayout.EditCategory(TEXT("ColorPainting"));
+
+	/** Creates a custom widget row containing all color channel flags */
+	TSharedRef<IPropertyHandle> RedChannel = DetailLayout.GetProperty("bWriteRed", UMeshColorPaintingToolProperties::StaticClass());
+	RedChannel->MarkHiddenByCustomization();
+	TSharedRef<IPropertyHandle> GreenChannel = DetailLayout.GetProperty("bWriteGreen", UMeshColorPaintingToolProperties::StaticClass());
+	GreenChannel->MarkHiddenByCustomization();
+	TSharedRef<IPropertyHandle> BlueChannel = DetailLayout.GetProperty("bWriteBlue", UMeshColorPaintingToolProperties::StaticClass());
+	BlueChannel->MarkHiddenByCustomization();
+	TSharedRef<IPropertyHandle> AlphaChannel = DetailLayout.GetProperty("bWriteAlpha", UMeshColorPaintingToolProperties::StaticClass());
+	AlphaChannel->MarkHiddenByCustomization();
+	TArray<TSharedRef<IPropertyHandle>> Channels = { RedChannel, GreenChannel, BlueChannel, AlphaChannel };
+	TSharedPtr<SHorizontalBox> ChannelsWidget;
+
+	ColorCategory.AddCustomRow(NSLOCTEXT("VertexPaintSettings", "ChannelLabel", "Channels"))
+		.NameContent()
+		[
+			SNew(STextBlock)
+			.Text(NSLOCTEXT("VertexPaintSettings", "ChannelsLabel", "Channels"))
+			.ToolTipText(NSLOCTEXT("VertexPaintSettings", "ChannelsToolTip", "Colors Channels which should be influenced during Painting."))
+		]
+		.ValueContent()
+		.MaxDesiredWidth(250.0f)
+		[
+			SAssignNew(ChannelsWidget, SHorizontalBox)
+		];
+
+	for (TSharedRef<IPropertyHandle> Channel : Channels)
+	{
+		ChannelsWidget->AddSlot()
+			.AutoWidth()
+			.Padding(0.0f, 0.0f, 4.0f, 0.0f)
+			[
+				CreateColorChannelWidget(Channel)
+			];
+	}
+
+	static FText RestrictReason = NSLOCTEXT("VertexPaintSettings", "TextureIndexRestriction", "Unable to paint this Texture, change Texture Weight Type");
+	BlendPaintEnumRestriction = MakeShareable(new FPropertyRestriction(RestrictReason));
+
+	/** Add custom row for painting on specific LOD level with callbacks to the painter to update the data */
+	TSharedRef<IPropertyHandle> LODPaintingEnabled = DetailLayout.GetProperty("bPaintOnSpecificLOD", UMeshColorPaintingToolProperties::StaticClass());
+	LODPaintingEnabled->MarkHiddenByCustomization();
+	TSharedRef<IPropertyHandle> LODPaintingIndex = DetailLayout.GetProperty("LODIndex", UMeshColorPaintingToolProperties::StaticClass());
+	LODPaintingIndex->MarkHiddenByCustomization();
+	TSharedPtr<SWidget> LODIndexWidget = LODPaintingIndex->CreatePropertyValueWidget();
+
+	ColorCategory.AddCustomRow(NSLOCTEXT("LODPainting", "LODPaintingLabel", "LOD Model Painting"))
+		.NameContent()
+		[
+			SNew(STextBlock)
+			.Text(NSLOCTEXT("LODPainting", "LODPaintingSetupLabel", "LOD Model Painting"))
+			.ToolTipText(NSLOCTEXT("LODPainting", "LODPaintingSetupToolTip", "Allows for Painting Vertex Colors on Specific LOD Models."))
+		]
+	.ValueContent()
+		[
+			SNew(SHorizontalBox)
+			+ SHorizontalBox::Slot()
+		.Padding(0.0f, 0.0f, 4.0f, 0.0f)
+		.AutoWidth()
+		[
+			SNew(SCheckBox)
+				.IsChecked_Lambda([=]() -> ECheckBoxState { return (UMeshPaintMode::GetColorToolProperties()->bPaintOnSpecificLOD ? ECheckBoxState::Checked : ECheckBoxState::Unchecked); })
+				.OnCheckStateChanged(FOnCheckStateChanged::CreateLambda([=](ECheckBoxState State) { 
+					if (UMeshColorPaintingTool* ColorBrush = Cast<UMeshColorPaintingTool>(UMeshPaintMode::GetMeshPaintMode()->GetToolManager()->GetActiveTool(EToolSide::Left)))
+					{
+						ColorBrush->LODPaintStateChanged(State == ECheckBoxState::Checked);
+					}
+				}))
+		]
+	+ SHorizontalBox::Slot()
+		.Padding(0.0f, 0.0f, 4.0f, 0.0f)
+		[
+			SNew(SNumericEntryBox<int32>)
+			.IsEnabled_Lambda([=]() -> bool { return UMeshPaintMode::GetColorToolProperties()->bPaintOnSpecificLOD;  })
+			.AllowSpin(true)
+			.Value_Lambda([=]() -> int32 { return UMeshPaintMode::GetColorToolProperties()->LODIndex; })
+			.MinValue(0)
+			.MaxValue_Lambda([=]() -> int32 { 
+					if (UMeshColorPaintingTool* ColorBrush = Cast<UMeshColorPaintingTool>(UMeshPaintMode::GetMeshPaintMode()->GetToolManager()->GetActiveTool(EToolSide::Left)))
+					{
+						return ColorBrush->GetMaxLODIndexToPaint();
+					}
+					return INT_MAX;
+				})
+			.MaxSliderValue_Lambda([=]() -> int32 { 
+					if (UMeshColorPaintingTool* ColorBrush = Cast<UMeshColorPaintingTool>(UMeshPaintMode::GetMeshPaintMode()->GetToolManager()->GetActiveTool(EToolSide::Left)))
+					{
+						return ColorBrush->GetMaxLODIndexToPaint();
+					}
+					return INT_MAX;
+				})
+			.OnValueChanged(SNumericEntryBox<int32>::FOnValueChanged::CreateLambda([=](int32 Value) { UMeshPaintMode::GetColorToolProperties()->LODIndex = Value; }))
+			.OnValueCommitted(SNumericEntryBox<int32>::FOnValueCommitted::CreateLambda([=](int32 Value, ETextCommit::Type CommitType) { 
+					UMeshPaintMode::GetColorToolProperties()->LODIndex = Value; 
+					if (UMeshColorPaintingTool* ColorBrush = Cast<UMeshColorPaintingTool>(UMeshPaintMode::GetMeshPaintMode()->GetToolManager()->GetActiveTool(EToolSide::Left)))
+					{
+						ColorBrush->PaintLODChanged();
+					}
+				}))
+		]
+		];
+
+	ColorCategory.AddCustomRow(NSLOCTEXT("LODPainting", "LODPaintingLabel", "LOD Model Painting"))
+		.WholeRowContent()
+		[
+			SNew(SBorder)
+			.Visibility_Lambda([this]() -> EVisibility
+			{
+				if (UMeshColorPaintingToolProperties* ColorProperties = UMeshPaintMode::GetColorToolProperties())
+				{
+					return ColorProperties->bPaintOnSpecificLOD ? EVisibility::Collapsed : EVisibility::Visible;
+				}
+				return EVisibility::Collapsed;
+			})
+		.Padding(FMargin(4.0f))
+		.BorderImage(FEditorStyle::GetBrush("SettingsEditor.CheckoutWarningBorder"))
+		.BorderBackgroundColor(FColor(166, 137, 0))
+		[
+			SNew(STextBlock)
+			.AutoWrapText(true)
+			.Font(IDetailLayoutBuilder::GetDetailFont())
+
+		.Text_Lambda([this]() -> FText
+		{
+			static const FText SkelMeshNotificationText = LOCTEXT("SkelMeshAssetPaintInfo", "Paint is propagated to Skeletal Mesh Asset(s)");
+			static const FText StaticMeshNotificationText = LOCTEXT("StaticMeshAssetPaintInfo", "Paint is applied to all LODs");
+
+			const bool bSkelMeshText = UMeshPaintMode::GetMeshPaintMode()->GetSelectedComponents<USkeletalMeshComponent>().Num() > 0;
+			const bool bLODPaintText = UMeshPaintMode::GetColorToolProperties() ?  !UMeshPaintMode::GetColorToolProperties()->bPaintOnSpecificLOD : false;
+			return FText::Format(FTextFormat::FromString(TEXT("{0}{1}{2}")), bSkelMeshText ? SkelMeshNotificationText : FText::GetEmpty(), bSkelMeshText && bLODPaintText ? FText::FromString(TEXT("\n")) : FText::GetEmpty(), bLODPaintText ? StaticMeshNotificationText : FText::GetEmpty());
+		})
+		]
+		];
+}
+
+TSharedRef<IDetailCustomization> FWeightPaintingSettingsCustomization::MakeInstance()
+{
+	return MakeShareable(new FWeightPaintingSettingsCustomization);
+}
+
+void FWeightPaintingSettingsCustomization::CustomizeDetails(IDetailLayoutBuilder& DetailLayout)
+{
+	FVertexPaintingSettingsCustomization::CustomizeDetails(DetailLayout);
+
+	TMap<FName, TSharedRef<IPropertyHandle>> CustomizedProperties;
+	TMap<FName, TSharedRef<IPropertyHandle>> Properties;
+
+	TSharedRef<IPropertyHandle> WeightTypeProperty = DetailLayout.GetProperty(GET_MEMBER_NAME_CHECKED(UMeshWeightPaintingToolProperties, TextureWeightType));
+	TSharedRef<IPropertyHandle> PaintWeightProperty = DetailLayout.GetProperty(GET_MEMBER_NAME_CHECKED(UMeshWeightPaintingToolProperties, PaintTextureWeightIndex));
+	TSharedRef<IPropertyHandle> EraseWeightProperty = DetailLayout.GetProperty(GET_MEMBER_NAME_CHECKED(UMeshWeightPaintingToolProperties, EraseTextureWeightIndex));
+
+	WeightTypeProperty->SetOnPropertyValueChanged(FSimpleDelegate::CreateRaw(this, &FWeightPaintingSettingsCustomization::OnTextureWeightTypeChanged, WeightTypeProperty, PaintWeightProperty, EraseWeightProperty));
+
+	static FText RestrictReason = NSLOCTEXT("VertexPaintSettings", "TextureIndexRestriction", "Unable to paint this Texture, change Texture Weight Type");
+	BlendPaintEnumRestriction = MakeShareable(new FPropertyRestriction(RestrictReason));
+	PaintWeightProperty->AddRestriction(BlendPaintEnumRestriction.ToSharedRef());
+	EraseWeightProperty->AddRestriction(BlendPaintEnumRestriction.ToSharedRef());
+
+	OnTextureWeightTypeChanged(WeightTypeProperty, PaintWeightProperty, EraseWeightProperty);
+}
+
+void FWeightPaintingSettingsCustomization::OnTextureWeightTypeChanged(TSharedRef<IPropertyHandle> WeightTypeProperty, TSharedRef<IPropertyHandle> PaintWeightProperty, TSharedRef<IPropertyHandle> EraseWeightProperty)
+{
+	UEnum* ImportTypeEnum = StaticEnum<EMeshPaintTextureIndex>();
+	uint8 EnumValue = 0;
+	WeightTypeProperty->GetValue(EnumValue);
+
+	BlendPaintEnumRestriction->RemoveAll();
+	for (uint8 EnumIndex = 0; EnumIndex < (ImportTypeEnum->GetMaxEnumValue() + 1); ++EnumIndex)
+	{
+		if ((EnumIndex + 1) > EnumValue)
+		{
+			FString EnumName = ImportTypeEnum->GetNameByValue(EnumIndex).ToString();
+			EnumName.RemoveFromStart("EMeshPaintTextureIndex::");
+			BlendPaintEnumRestriction->AddDisabledValue(EnumName);
+		}
+	}
+
+	uint8 Value = 0;
+	PaintWeightProperty->GetValue(Value);
+	Value = FMath::Clamp<uint8>(Value, 0, EnumValue - 1);
+	PaintWeightProperty->SetValue(Value);
+
+	EraseWeightProperty->GetValue(Value);
+	Value = FMath::Clamp<uint8>(Value, 0, EnumValue - 1);
+	EraseWeightProperty->SetValue(Value);
+}
+
+#undef LOCTEXT_NAMESPACE
\ No newline at end of file
diff --git a/Engine/Plugins/MeshPainting/Source/MeshPaintEditorMode/Private/BrushSettingsCustomization.h b/Engine/Plugins/MeshPainting/Source/MeshPaintEditorMode/Private/BrushSettingsCustomization.h
new file mode 100644
--- /dev/null
+++ b/Engine/Plugins/MeshPainting/Source/MeshPaintEditorMode/Private/BrushSettingsCustomization.h
@@ -0,0 +1,71 @@
+// Copyright 1998-2019 Epic Games, Inc. All Rights Reserved.
+
+#pragma once
+
+#include "IDetailCustomization.h"
+#include "IPropertyTypeCustomization.h"
+#include "IDetailRootObjectCustomization.h"
+#include "PropertyRestriction.h"
+#include "Widgets/SWidget.h"
+
+class SHorizontalBox;
+class UMeshVertexPaintingToolProperties;
+
+
+class FVertexPaintingSettingsCustomization : public IDetailCustomization
+{
+public:
+	static TSharedRef<IDetailCustomization> MakeInstance();
+
+	/** IPropertyTypeCustomization interface */
+	virtual void CustomizeDetails(IDetailLayoutBuilder& DetailLayout) override;
+protected:
+	/** Callback for when texture weight type changed so we can update restrictions */
+	void OnTextureWeightTypeChanged(TSharedRef<IPropertyHandle> WeightTypeProperty, TSharedRef<IPropertyHandle> PaintWeightProperty, TSharedRef<IPropertyHandle> EraseWeightProperty);
+	FReply OnSwapColorsClicked(TSharedRef<IPropertyHandle> PaintColor, TSharedRef<IPropertyHandle> EraseColor);
+
+	/** Property restriction applied to blend paint enum dropdown box */
+	TSharedPtr<FPropertyRestriction> BlendPaintEnumRestriction;
+};
+
+class FColorPaintingSettingsCustomization : public FVertexPaintingSettingsCustomization
+{
+public:
+	static TSharedRef<IDetailCustomization> MakeInstance();
+
+	/** IPropertyTypeCustomization interface */
+	virtual void CustomizeDetails(IDetailLayoutBuilder& DetailLayout) override;
+};
+
+class FWeightPaintingSettingsCustomization : public FVertexPaintingSettingsCustomization
+{
+public:
+	static TSharedRef<IDetailCustomization> MakeInstance();
+
+	/** IPropertyTypeCustomization interface */
+	virtual void CustomizeDetails(IDetailLayoutBuilder& DetailLayout) override;
+protected:
+	void OnTextureWeightTypeChanged(TSharedRef<IPropertyHandle> WeightTypeProperty, TSharedRef<IPropertyHandle> PaintWeightProperty, TSharedRef<IPropertyHandle> EraseWeightProperty);
+
+	/** Property restriction applied to blend paint enum dropdown box */
+	TSharedPtr<FPropertyRestriction> BlendPaintEnumRestriction;
+
+};
+
+// class FTexturePaintSettingsCustomization : public IPropertyTypeCustomization
+// {
+// public:
+// 	static TSharedRef<IPropertyTypeCustomization> MakeInstance();
+// 
+// 	/** IPropertyTypeCustomization interface */
+// 	virtual void CustomizeHeader(TSharedRef<IPropertyHandle> PropertyHandle, FDetailWidgetRow& HeaderRow, IPropertyTypeCustomizationUtils& CustomizationUtils) override {};
+// 	virtual void CustomizeChildren(TSharedRef<IPropertyHandle> PropertyHandle, IDetailChildrenBuilder& ChildBuilder, IPropertyTypeCustomizationUtils& CustomizationUtils) override;
+// protected:
+// 	/** Cached instance of the painter */
+// 	FPaintModePainter* MeshPainter;
+// 	/** Cached instance of the texture painting settings */
+// 	FTexturePaintSettings* PaintSettings;
+// };
+
+
+	
\ No newline at end of file
diff --git a/Engine/Plugins/MeshPainting/Source/MeshPaintEditorMode/Private/MeshPaintEditorMode.cpp b/Engine/Plugins/MeshPainting/Source/MeshPaintEditorMode/Private/MeshPaintEditorMode.cpp
new file mode 100644
--- /dev/null
+++ b/Engine/Plugins/MeshPainting/Source/MeshPaintEditorMode/Private/MeshPaintEditorMode.cpp
@@ -0,0 +1,147 @@
+// Copyright 1998-2019 Epic Games, Inc. All Rights Reserved.
+
+#include "MeshPaintEditorMode.h"
+#include "EditorModeRegistry.h"
+#include "Modules/ModuleManager.h"
+#include "MeshPaintMode.h"
+#include "EditorModeManager.h"
+#include "MeshPaintModeSettings.h"
+#include "MeshPaintModeCommands.h"
+#include "BrushSettingsCustomization.h"
+#include "PropertyEditorModule.h"
+#include "MeshPaintComponentAdapterFactory.h"
+#include "MeshPaintAdapterFactory.h"
+#include "MeshPaintSplineMeshAdapter.h"
+#include "MeshPaintStaticMeshAdapter.h"
+#include "MeshPaintSkeletalMeshAdapter.h"
+#include "Settings/LevelEditorMiscSettings.h"
+ 
+
+#define LOCTEXT_NAMESPACE "MeshPaintMode"
+
+class FMeshPaintEditorModeModule : public IModuleInterface
+{
+public:
+	FMeshPaintEditorModeModule() 
+	{
+	}
+
+	// FModuleInterface overrides
+	virtual void StartupModule() override;
+	virtual void ShutdownModule() override;
+	virtual bool SupportsDynamicReloading() override
+	{
+		return true;
+	}
+
+protected:
+
+	/** Stores the FEditorMode ID of the associated editor mode */
+	static inline FEditorModeID GetEditorModeID()
+	{
+		static FEditorModeID MeshPaintModeFeatureName = FName( TEXT( "MeshPaintMode" ) );
+		return MeshPaintModeFeatureName;
+	}
+
+	void Register();
+	void OnPostEngineInit();
+	void Unregister();
+
+
+	void OnMeshPaintModeButtonClicked();
+	bool IsMeshPaintModeButtonEnabled();
+	void RegisterGeometryAdapterFactory(TSharedRef<IMeshPaintComponentAdapterFactory> Factory);
+	void UnregisterGeometryAdapterFactory(TSharedRef<IMeshPaintComponentAdapterFactory> Factory);
+};
+
+void FMeshPaintEditorModeModule::Register()
+{
+	RegisterGeometryAdapterFactory(MakeShareable(new FMeshPaintSplineMeshComponentAdapterFactory));
+	RegisterGeometryAdapterFactory(MakeShareable(new FMeshPaintStaticMeshComponentAdapterFactory));
+	RegisterGeometryAdapterFactory(MakeShareable(new FMeshPaintSkeletalMeshComponentAdapterFactory));
+
+
+	FEditorModeRegistry::Get().RegisterScriptableMode<UMeshPaintMode>(
+		GetEditorModeID(),
+		LOCTEXT("ModeName", "Mesh Paint"),
+		FSlateIcon(FEditorStyle::GetStyleSetName(), "LevelEditor.MeshPaintMode", "LevelEditor.MeshPaintMode.Small"),
+		true,
+		600
+		);
+
+	FPropertyEditorModule& PropertyModule = FModuleManager::LoadModuleChecked<FPropertyEditorModule>("PropertyEditor");
+	PropertyModule.RegisterCustomClassLayout("MeshColorPaintingToolProperties", FOnGetDetailCustomizationInstance::CreateStatic(&FColorPaintingSettingsCustomization::MakeInstance));
+	PropertyModule.RegisterCustomClassLayout("MeshWeightPaintingToolProperties", FOnGetDetailCustomizationInstance::CreateStatic(&FWeightPaintingSettingsCustomization::MakeInstance));
+}
+
+
+void FMeshPaintEditorModeModule::OnPostEngineInit()
+{
+	if (!GetDefault<ULevelEditorMiscSettings>()->bEnableLegacyMeshPaintMode)
+	{
+		Register();
+		FMeshPaintingToolActionCommands::RegisterAllToolActions();
+		FMeshPaintEditorModeCommands::Register();
+	}
+}
+
+void FMeshPaintEditorModeModule::StartupModule()
+{
+	FCoreDelegates::OnPostEngineInit.AddRaw(this, &FMeshPaintEditorModeModule::OnPostEngineInit);
+}
+
+
+void FMeshPaintEditorModeModule::ShutdownModule()
+{
+	
+	FCoreDelegates::OnPostEngineInit.RemoveAll(this);
+
+	FMeshPaintingToolActionCommands::UnregisterAllToolActions();
+	FMeshPaintEditorModeCommands::Unregister();
+
+	// This function may be called during shutdown to clean up your module.  For modules that support dynamic reloading,
+	// we call this function before unloading the module.
+	Unregister();
+	
+}
+
+
+
+
+void FMeshPaintEditorModeModule::Unregister()
+{
+	FEditorModeRegistry::Get().UnregisterMode(GetEditorModeID());
+}
+
+
+void FMeshPaintEditorModeModule::OnMeshPaintModeButtonClicked()
+{
+	// *Important* - activate the mode first since FEditorModeTools::DeactivateMode will
+	// activate the default mode when the stack becomes empty, resulting in multiple active visible modes.
+	GLevelEditorModeTools().ActivateMode(GetEditorModeID());
+
+	// Find and disable any other 'visible' modes since we only ever allow one of those active at a time.
+	GLevelEditorModeTools().DeactivateOtherVisibleModes(GetEditorModeID());
+
+}
+
+
+
+bool FMeshPaintEditorModeModule::IsMeshPaintModeButtonEnabled( )
+{
+	return true;
+}
+
+void  FMeshPaintEditorModeModule::RegisterGeometryAdapterFactory(TSharedRef<IMeshPaintComponentAdapterFactory> Factory)
+{
+	FMeshPaintComponentAdapterFactory::FactoryList.Add(Factory);
+}
+
+void  FMeshPaintEditorModeModule::UnregisterGeometryAdapterFactory(TSharedRef<IMeshPaintComponentAdapterFactory> Factory)
+{
+	FMeshPaintComponentAdapterFactory::FactoryList.Remove(Factory);
+}
+
+IMPLEMENT_MODULE(FMeshPaintEditorModeModule, MeshPaintEditorMode)
+
+#undef LOCTEXT_NAMESPACE
