Copy up from Dev-Editor [REDACTED_HANDLE]

diff --git a/Engine/Plugins/MeshPainting/Source/MeshPaintEditorMode/Private/MeshPaintMode.cpp b/Engine/Plugins/MeshPainting/Source/MeshPaintEditorMode/Private/MeshPaintMode.cpp
new file mode 100644
--- /dev/null
+++ b/Engine/Plugins/MeshPainting/Source/MeshPaintEditorMode/Private/MeshPaintMode.cpp
@@ -0,0 +1,704 @@
+// Copyright 1998-2019 Epic Games, Inc. All Rights Reserved.
+
+#include "MeshPaintMode.h"
+#include "MeshPaintModeCommands.h"
+#include "MeshPaintModeToolkit.h"
+#include "Toolkits/ToolkitManager.h"
+#include "Framework/Commands/UICommandList.h"
+#include "MeshPaintModeSettings.h"
+#include "Framework/Commands/Commands.h"
+#include "MultiSelectionTool.h"
+#include "InteractiveToolManager.h"
+#include "InteractiveToolsContext.h"
+#include "Components/MeshComponent.h"
+#include "Dialogs/Dialogs.h"
+#include "Components/StaticMeshComponent.h"
+#include "MeshPaintHelpers.h"
+#include "Components/SkeletalMeshComponent.h"
+#include "ScopedTransaction.h"
+#include "PackageTools.h"
+#include "MeshPaintHelpers.h"
+#include "Engine/Selection.h"
+#include "MeshPaintHelpers.h"
+#include "MeshPaintAdapterFactory.h"
+#include "MeshPaintModeHelpers.h"
+#include "MeshSelect.h"
+
+
+#define LOCTEXT_NAMESPACE "MeshPaintMode"
+
+FName UMeshPaintMode::MeshPaintMode_Color = FName(TEXT("Color"));
+FName UMeshPaintMode::MeshPaintMode_Weights = FName(TEXT("Weights"));
+FName UMeshPaintMode::MeshPaintMode_Texture = FName(TEXT("Texture"));
+
+UMeshVertexPaintingToolProperties* UMeshPaintMode::GetVertexToolProperties()
+{
+	if (GLevelEditorModeTools().GetActiveScriptableMode("MeshPaintMode")
+		&& GLevelEditorModeTools().GetActiveScriptableMode("MeshPaintMode")->GetToolManager()->GetActiveTool(EToolSide::Mouse))
+	{
+		TArray<UObject*> PropertyArray = GLevelEditorModeTools().GetActiveScriptableMode("MeshPaintMode")->GetToolManager()->GetActiveTool(EToolSide::Mouse)->GetToolProperties();
+		for(UObject* Property : PropertyArray)
+		{
+			if (UMeshVertexPaintingToolProperties* VertexProperties = Cast<UMeshVertexPaintingToolProperties>(Property))
+			{
+				return VertexProperties;
+			}
+		}
+	}
+	return nullptr;
+}
+
+UMeshColorPaintingToolProperties* UMeshPaintMode::GetColorToolProperties()
+{
+	if (GLevelEditorModeTools().GetActiveScriptableMode("MeshPaintMode")
+		&& GLevelEditorModeTools().GetActiveScriptableMode("MeshPaintMode")->GetToolManager()->GetActiveTool(EToolSide::Mouse))
+	{
+		TArray<UObject*> PropertyArray = GLevelEditorModeTools().GetActiveScriptableMode("MeshPaintMode")->GetToolManager()->GetActiveTool(EToolSide::Mouse)->GetToolProperties();
+		for (UObject* Property : PropertyArray)
+		{
+			if (UMeshColorPaintingToolProperties* ColorProperties = Cast<UMeshColorPaintingToolProperties>(Property))
+			{
+				return ColorProperties;
+			}
+		}
+	}
+	return nullptr;
+}
+
+UMeshWeightPaintingToolProperties* UMeshPaintMode::GetWeightToolProperties()
+{
+	if (GLevelEditorModeTools().GetActiveScriptableMode("MeshPaintMode")
+		&& GLevelEditorModeTools().GetActiveScriptableMode("MeshPaintMode")->GetToolManager()->GetActiveTool(EToolSide::Mouse))
+	{
+		TArray<UObject*> PropertyArray = GLevelEditorModeTools().GetActiveScriptableMode("MeshPaintMode")->GetToolManager()->GetActiveTool(EToolSide::Mouse)->GetToolProperties();
+		for (UObject* Property : PropertyArray)
+		{
+			if (UMeshWeightPaintingToolProperties* WeightProperties = Cast<UMeshWeightPaintingToolProperties>(Property))
+			{
+				return WeightProperties;
+			}
+		}
+	}
+	return nullptr;
+}
+
+UMeshPaintMode* UMeshPaintMode::GetMeshPaintMode()
+{
+	return Cast<UMeshPaintMode>(GLevelEditorModeTools().GetActiveScriptableMode("MeshPaintMode"));
+}
+
+UMeshPaintMode::UMeshPaintMode()
+	:Super()
+{
+	SettingsClass = UMeshPaintModeSettings::StaticClass();
+	ToolsContextClass = UMeshToolsContext::StaticClass();
+	CurrentPaletteName = MeshPaintMode_Color;
+}
+
+void UMeshPaintMode::Enter()
+{
+	Super::Enter();
+	GEditor->OnEditorClose().AddUObject(this, &UMeshPaintMode::OnResetViewMode);
+	ModeSettings = Cast<UMeshPaintModeSettings>(SettingsObject);
+	FMeshPaintEditorModeCommands ToolManagerCommands = FMeshPaintEditorModeCommands::Get();
+	RegisterTool(ToolManagerCommands.Select, TEXT("MeshClickTool"), NewObject<UMeshClickToolBuilder>());
+	RegisterTool(ToolManagerCommands.ColorPaint, TEXT("ColorBrushTool"), NewObject<UMeshColorPaintingToolBuilder>());
+	RegisterTool(ToolManagerCommands.WeightPaint, TEXT("WeightBrushTool"), NewObject<UMeshWeightPaintingToolBuilder>());
+
+	UpdateSelectedMeshes();
+
+	ActivateDefaultTool();
+}
+
+void UMeshPaintMode::Exit()
+{
+	GEditor->OnEditorClose().RemoveAll(this);
+	OnResetViewMode();
+	const FMeshPaintEditorModeCommands& Commands = FMeshPaintEditorModeCommands::Get();
+	const TSharedRef<FUICommandList>& CommandList = Toolkit->GetToolkitCommands();
+	for (const TSharedPtr<const FUICommandInfo> Action : Commands.Commands[UMeshPaintMode::MeshPaintMode_Color])
+	{
+		CommandList->UnmapAction(Action);
+	}
+	for (const TSharedPtr<const FUICommandInfo> Action : Commands.Commands[UMeshPaintMode::MeshPaintMode_Weights])
+	{
+		CommandList->UnmapAction(Action);
+	}
+	for (const TSharedPtr<const FUICommandInfo> Action : Commands.Commands[UMeshPaintMode::MeshPaintMode_Weights])
+	{
+		CommandList->UnmapAction(Action);
+	}
+	Super::Exit();
+}
+
+void UMeshPaintMode::CreateToolkit()
+{
+	if (!Toolkit.IsValid())
+	{
+		Toolkit = MakeShareable(new FMeshPaintModeToolkit);
+		Toolkit->Init(Owner->GetToolkitHost());
+	}
+
+	// Register UI commands
+	BindCommands();
+
+	Super::CreateToolkit();
+}
+
+void UMeshPaintMode::Tick(FEditorViewportClient* ViewportClient, float DeltaTime)
+{
+	UEdMode::Tick(ViewportClient, DeltaTime);
+
+	if (ViewportClient->IsPerspective())
+	{
+		// Make sure perspective viewports are still set to real-time
+		UMeshPaintModeHelpers::SetRealtimeViewport(true);
+
+		// Set viewport show flags		
+		UMeshPaintModeHelpers::SetViewportColorMode(ModeSettings->ColorViewMode, ViewportClient);
+	}
+
+
+	if (bRecacheVertexDataSize)
+	{
+		UpdateCachedVertexDataSize();
+	}
+}
+
+TMap<FName, TArray<TSharedPtr<FUICommandInfo>>> UMeshPaintMode::GetModeCommands() const
+{
+	return FMeshPaintEditorModeCommands::GetCommands();
+}
+
+void UMeshPaintMode::BindCommands()
+{
+	const FMeshPaintEditorModeCommands& Commands = FMeshPaintEditorModeCommands::Get();
+	const TSharedRef<FUICommandList>& CommandList = Toolkit->GetToolkitCommands();
+
+	auto IsInValidPaintMode = [this]() -> bool { return (GetVertexToolProperties() != nullptr) && Cast<UMeshToolManager>(GetToolManager())->SelectionContainsValidAdapters(); };
+	CommandList->MapAction(Commands.Fill, 
+		FUIAction(FExecuteAction::CreateUObject(this, &UMeshPaintMode::FillWithVertexColor),
+		FCanExecuteAction::CreateLambda(IsInValidPaintMode)));
+
+	CommandList->MapAction(Commands.Propagate,
+		FUIAction(
+			FExecuteAction::CreateUObject(this, &UMeshPaintMode::PropagateVertexColorsToAsset),
+			FCanExecuteAction::CreateUObject(this, &UMeshPaintMode::CanPropagateVertexColors)
+		));
+
+	auto IsAValidMeshComponentSelected = [this]() -> bool { return (GetSelectedComponents<UMeshComponent>().Num() == 1) && Cast<UMeshToolManager>(GetToolManager())->SelectionContainsValidAdapters(); };
+	CommandList->MapAction(Commands.Import,
+		FUIAction(
+			FExecuteAction::CreateUObject(this, &UMeshPaintMode::ImportVertexColors),
+			FCanExecuteAction::CreateLambda(IsAValidMeshComponentSelected)
+		));
+
+	CommandList->MapAction(Commands.Save,
+		FUIAction(
+			FExecuteAction::CreateUObject(this, &UMeshPaintMode::SavePaintedAssets),
+			FCanExecuteAction::CreateUObject(this, &UMeshPaintMode::CanSaveMeshPackages)
+		));
+
+	CommandList->MapAction(Commands.Copy,
+		FUIAction(
+			FExecuteAction::CreateUObject(this, &UMeshPaintMode::CopyVertexColors),
+			FCanExecuteAction::CreateUObject(this, &UMeshPaintMode::CanCopyInstanceVertexColors)
+		));
+
+	CommandList->MapAction(Commands.Paste,
+		FUIAction(
+			FExecuteAction::CreateUObject(this, &UMeshPaintMode::PasteVertexColors),
+			FCanExecuteAction::CreateUObject(this, &UMeshPaintMode::CanPasteInstanceVertexColors)
+		));
+
+	CommandList->MapAction(Commands.Remove,
+		FUIAction(
+			FExecuteAction::CreateUObject(this, &UMeshPaintMode::RemoveVertexColors),
+			FCanExecuteAction::CreateUObject(this, &UMeshPaintMode::CanRemoveInstanceColors)
+		));
+
+	CommandList->MapAction(Commands.Fix,
+		FUIAction(
+			FExecuteAction::CreateUObject(this, &UMeshPaintMode::FixVertexColors),
+			FCanExecuteAction::CreateUObject(this, &UMeshPaintMode::DoesRequireVertexColorsFixup)
+		));
+
+	CommandList->MapAction(Commands.PropagateVertexColorsToLODs,
+		FUIAction(
+			FExecuteAction::CreateUObject(this, &UMeshPaintMode::PropagateVertexColorsToLODs),
+			FCanExecuteAction::CreateUObject(this, &UMeshPaintMode::CanPropagateVertexColorsToLODs)
+		));
+
+}
+
+
+
+
+void UMeshPaintMode::OnToolStarted(UInteractiveToolManager* Manager, UInteractiveTool* Tool)
+{
+	FMeshPaintingToolActionCommands::UpdateToolCommandBinding(Tool, ToolCommandList, false);
+
+	if (UMeshVertexPaintingTool* VertexPaintingTool = Cast<UMeshVertexPaintingTool>(GetToolManager()->GetActiveTool(EToolSide::Left)))
+	{
+		VertexPaintingTool->OnPaintingFinished().BindUObject(this, &UMeshPaintMode::UpdateCachedVertexDataSize);
+	}
+}
+
+void UMeshPaintMode::OnToolEnded(UInteractiveToolManager* Manager, UInteractiveTool* Tool)
+{
+	FMeshPaintingToolActionCommands::UpdateToolCommandBinding(Tool, ToolCommandList, true);
+	// First update your bindings, then call the base behavior
+	Super::OnToolEnded(Manager, Tool);
+}
+
+bool UMeshPaintMode::InputAxis(FEditorViewportClient* InViewportClient, FViewport* Viewport, int32 ControllerId, FKey Key, float Delta, float DeltaTime)
+{
+	bool bSelectingOnly = (GetToolManager()->GetActiveTool(EToolSide::Mouse)) ? GetToolManager()->GetActiveTool(EToolSide::Mouse)->IsA<USingleClickTool>() : false;
+
+	// If we are painting (not selecting), don't allow panning with the mouse.
+	return !bSelectingOnly && !InViewportClient->IsAltPressed();
+}
+
+void UMeshPaintMode::ActorSelectionChangeNotify()
+{
+	UpdateSelectedMeshes();
+}
+
+void UMeshPaintMode::UpdateSelectedMeshes()
+{
+	if (UMeshToolManager* MeshToolManager = Cast<UMeshToolManager>(GetToolManager()))
+	{
+		MeshToolManager->ResetState();
+		const TArray<UMeshComponent*> CurrentMeshComponents = GetSelectedComponents<UMeshComponent>();
+		MeshToolManager->AddSelectedMeshComponents(CurrentMeshComponents);
+		MeshToolManager->bNeedsRecache = true;
+	}
+	bRecacheVertexDataSize = true;
+}
+
+void UMeshPaintMode::FillWithVertexColor()
+{
+	FScopedTransaction Transaction(LOCTEXT("LevelMeshPainter_TransactionFillInstColors", "Filling Per-Instance Vertex Colors"));
+	const TArray<UMeshComponent*> MeshComponents = GetSelectedComponents<UMeshComponent>();
+
+	static const bool bConvertSRGB = false;
+	FColor FillColor = FColor::White;
+	FColor MaskColor = FColor::White;
+
+	if (GetToolManager()->GetActiveTool(EToolSide::Mouse)->IsA<UMeshWeightPaintingTool>())
+	{
+		FillColor = UMeshPaintingToolset::GenerateColorForTextureWeight((int32)GetWeightToolProperties()->TextureWeightType, (int32)GetWeightToolProperties()->PaintTextureWeightIndex).ToFColor(bConvertSRGB);
+	}
+	else if (UMeshColorPaintingToolProperties* ColorProperties = GetColorToolProperties())
+	{
+		FillColor = ColorProperties->PaintColor.ToFColor(bConvertSRGB);
+		MaskColor.R = ColorProperties->bWriteRed ? 255 : 0;
+		MaskColor.G = ColorProperties->bWriteGreen ? 255 : 0;
+		MaskColor.B = ColorProperties->bWriteBlue ? 255 : 0;
+		MaskColor.A = ColorProperties->bWriteAlpha ? 255 : 0;
+	}
+
+	TUniquePtr< FComponentReregisterContext > ComponentReregisterContext;
+	/** Fill each mesh component with the given vertex color */
+	for (UMeshComponent* Component : MeshComponents)
+	{
+		checkf(Component != nullptr, TEXT("Invalid Mesh Component"));
+		Component->Modify();
+		ComponentReregisterContext = MakeUnique<FComponentReregisterContext>(Component);
+
+		TSharedPtr<IMeshPaintComponentAdapter> MeshAdapter = Cast<UMeshToolManager>(GetToolManager())->GetAdapterForComponent(Component);
+		if (MeshAdapter)
+		{
+			MeshAdapter->PreEdit();
+		}
+
+		if (UMeshColorPaintingToolProperties* ColorProperties = UMeshPaintMode::GetColorToolProperties())
+		{
+			const bool bPaintOnSpecificLOD = ColorProperties ? ColorProperties->bPaintOnSpecificLOD : false;
+
+			if (Component->IsA<UStaticMeshComponent>())
+			{
+				UMeshPaintingToolset::FillStaticMeshVertexColors(Cast<UStaticMeshComponent>(Component), bPaintOnSpecificLOD ? ColorProperties->LODIndex : -1, FillColor, MaskColor);
+			}
+			else if (Component->IsA<USkeletalMeshComponent>())
+			{
+				UMeshPaintingToolset::FillSkeletalMeshVertexColors(Cast<USkeletalMeshComponent>(Component), bPaintOnSpecificLOD ? ColorProperties->LODIndex : -1, FillColor, MaskColor);
+			}
+		}
+
+		if (MeshAdapter)
+		{
+			MeshAdapter->PostEdit();
+		}
+	}
+}
+
+
+void UMeshPaintMode::PropagateVertexColorsToAsset()
+{
+	const TArray<UStaticMeshComponent*> StaticMeshComponents = GetSelectedComponents<UStaticMeshComponent>();
+	FSuppressableWarningDialog::FSetupInfo SetupInfo(LOCTEXT("PushInstanceVertexColorsPrompt_Message", "Copying the instance vertex colors to the source mesh will replace any of the source mesh's pre-existing vertex colors and affect every instance of the source mesh."),
+		LOCTEXT("PushInstanceVertexColorsPrompt_Title", "Warning: Copying vertex data overwrites all instances"), "Warning_PushInstanceVertexColorsPrompt");
+
+	SetupInfo.ConfirmText = LOCTEXT("PushInstanceVertexColorsPrompt_ConfirmText", "Continue");
+	SetupInfo.CancelText = LOCTEXT("PushInstanceVertexColorsPrompt_CancelText", "Abort");
+	SetupInfo.CheckBoxText = LOCTEXT("PushInstanceVertexColorsPrompt_CheckBoxText", "Always copy vertex colors without prompting");
+
+	FSuppressableWarningDialog VertexColorCopyWarning(SetupInfo);
+
+	// Prompt the user to see if they really want to push the vert colors to the source mesh and to explain
+	// the ramifications of doing so. This uses a suppressible dialog so that the user has the choice to always ignore the warning.
+	if (VertexColorCopyWarning.ShowModal() != FSuppressableWarningDialog::Cancel)
+	{
+		FScopedTransaction Transaction(LOCTEXT("LevelMeshPainter_TransactionPropogateColors", "Propagating Vertex Colors To Source Meshes"));
+		UMeshPaintModeHelpers::PropagateVertexColors(StaticMeshComponents);
+	}
+}
+
+bool UMeshPaintMode::CanPropagateVertexColors() const
+{
+	// Check whether or not our selected Static Mesh Components contain instance based vertex colors (only these can be propagated to the base mesh)
+	int32 NumInstanceVertexColorBytes = 0;
+
+	TArray<UStaticMesh*> StaticMeshes;
+	TArray<UStaticMeshComponent*> StaticMeshComponents = GetSelectedComponents<UStaticMeshComponent>();
+	return UMeshPaintModeHelpers::CanPropagateVertexColors(StaticMeshComponents, StaticMeshes, NumInstanceVertexColorBytes);
+
+}
+
+void UMeshPaintMode::ImportVertexColors()
+{
+	const TArray<UMeshComponent*> MeshComponents = GetSelectedComponents<UMeshComponent>();
+	if (MeshComponents.Num() == 1)
+	{
+		/** Import vertex color to single selected mesh component */
+		FScopedTransaction Transaction(LOCTEXT("LevelMeshPainter_TransactionImportColors", "Importing Vertex Colors From Texture"));
+		UMeshPaintModeHelpers::ImportVertexColorsFromTexture(MeshComponents[0]);
+	}
+}
+
+void UMeshPaintMode::SavePaintedAssets()
+{
+	const TArray<UStaticMeshComponent*> StaticMeshComponents = GetSelectedComponents<UStaticMeshComponent>();
+	const TArray<USkeletalMeshComponent*> SkeletalMeshComponents = GetSelectedComponents<USkeletalMeshComponent>();
+
+	/** Try and save outstanding dirty packages for currently selected mesh components */
+	TArray<UObject*> ObjectsToSave;
+	for (UStaticMeshComponent* StaticMeshComponent : StaticMeshComponents)
+	{
+		if (StaticMeshComponent && StaticMeshComponent->GetStaticMesh())
+		{
+			ObjectsToSave.Add(StaticMeshComponent->GetStaticMesh());
+		}
+	}
+
+	for (USkeletalMeshComponent* SkeletalMeshComponent : SkeletalMeshComponents)
+	{
+		if (SkeletalMeshComponent && SkeletalMeshComponent->SkeletalMesh)
+		{
+			ObjectsToSave.Add(SkeletalMeshComponent->SkeletalMesh);
+		}
+	}
+
+	if (ObjectsToSave.Num() > 0)
+	{
+		UPackageTools::SavePackagesForObjects(ObjectsToSave);
+	}
+}
+
+bool UMeshPaintMode::CanSaveMeshPackages() const
+{
+	// Check whether or not any of our selected mesh components contain mesh objects which require saving
+	TArray<UMeshComponent*> Components = GetSelectedComponents<UMeshComponent>();
+
+	bool bValid = false;
+
+	for (UMeshComponent* Component : Components)
+	{
+		UObject* Object = nullptr;
+		if (UStaticMeshComponent* StaticMeshComponent = Cast<UStaticMeshComponent>(Component))
+		{
+			Object = StaticMeshComponent->GetStaticMesh();
+		}
+		else if (USkeletalMeshComponent* SkeletalMeshComponent = Cast<USkeletalMeshComponent>(Component))
+		{
+			Object = SkeletalMeshComponent->SkeletalMesh;
+		}
+
+		if (Object != nullptr && Object->GetOutermost()->IsDirty())
+		{
+			bValid = true;
+			break;
+		}
+	}
+
+	return bValid;
+}
+
+bool UMeshPaintMode::CanRemoveInstanceColors() const
+{
+	const TArray<UStaticMeshComponent*> StaticMeshComponents = GetSelectedComponents<UStaticMeshComponent>();
+	int32 PaintingMeshLODIndex = 0;
+	if (UMeshColorPaintingToolProperties* ColorProperties = UMeshPaintMode::GetColorToolProperties())
+	{
+		PaintingMeshLODIndex = ColorProperties->bPaintOnSpecificLOD ? ColorProperties->LODIndex : 0;
+	}
+	int32 NumValidMeshes = 0;
+	// Retrieve per instance vertex color information (only valid if the component contains actual instance vertex colors)
+	for (UStaticMeshComponent* Component : StaticMeshComponents)
+	{
+		if (Component != nullptr && Component->GetStaticMesh() != nullptr && Component->GetStaticMesh()->GetNumLODs() > (int32)PaintingMeshLODIndex)
+		{
+			uint32 BufferSize = UMeshPaintingToolset::GetVertexColorBufferSize(Component, PaintingMeshLODIndex, true);
+
+			if (BufferSize > 0)
+			{
+				++NumValidMeshes;
+			}
+		}
+	}
+
+	return (NumValidMeshes != 0);
+}
+
+bool UMeshPaintMode::CanPasteInstanceVertexColors() const
+{
+	const TArray<UStaticMeshComponent*> StaticMeshComponents = GetSelectedComponents<UStaticMeshComponent>();
+	const TArray<FPerComponentVertexColorData> CopiedColorsByComponent = Cast<UMeshToolManager>(GetToolManager())->GetCopiedColorsByComponent();
+	return UMeshPaintModeHelpers::CanPasteInstanceVertexColors(StaticMeshComponents, CopiedColorsByComponent);
+}
+
+bool UMeshPaintMode::CanCopyInstanceVertexColors() const
+{
+	const TArray<UStaticMeshComponent*> StaticMeshComponents = GetSelectedComponents<UStaticMeshComponent>();
+	int32 PaintingMeshLODIndex = 0;
+	if (UMeshColorPaintingToolProperties* ColorProperties = UMeshPaintMode::GetColorToolProperties())
+	{
+		PaintingMeshLODIndex = ColorProperties->bPaintOnSpecificLOD ? ColorProperties->LODIndex : 0;
+	}
+
+	return UMeshPaintModeHelpers::CanCopyInstanceVertexColors(StaticMeshComponents, PaintingMeshLODIndex);
+
+}
+
+bool UMeshPaintMode::CanPropagateVertexColorsToLODs() const
+{
+	bool bPaintOnSpecificLOD = false;
+	if (UMeshColorPaintingToolProperties* ColorProperties = UMeshPaintMode::GetColorToolProperties())
+	{
+		bPaintOnSpecificLOD = ColorProperties ? ColorProperties->bPaintOnSpecificLOD : false;
+	}
+	// Can propagate when the mesh contains per-lod vertex colors or when we are not painting to a specific lod
+	const bool bSelectionContainsPerLODColors = Cast<UMeshToolManager>(GetToolManager())->SelectionContainsPerLODColors();
+	return bSelectionContainsPerLODColors || bPaintOnSpecificLOD;
+}
+
+void UMeshPaintMode::CopyVertexColors()
+{
+	const TArray<UStaticMeshComponent*> StaticMeshComponents = GetSelectedComponents<UStaticMeshComponent>();
+	TArray<FPerComponentVertexColorData> CopiedColorsByComponent;
+	UMeshPaintModeHelpers::CopyVertexColors(StaticMeshComponents, CopiedColorsByComponent);
+	Cast<UMeshToolManager>(GetToolManager())->SetCopiedColorsByComponent(CopiedColorsByComponent);
+
+}
+
+void UMeshPaintMode::PasteVertexColors()
+{
+	FScopedTransaction Transaction(LOCTEXT("LevelMeshPainter_TransactionPasteInstColors", "Pasting Per-Instance Vertex Colors"));
+	const TArray<UStaticMeshComponent*> StaticMeshComponents = GetSelectedComponents<UStaticMeshComponent>();
+	TArray<FPerComponentVertexColorData> CopiedColorsByComponent = Cast<UMeshToolManager>(GetToolManager())->GetCopiedColorsByComponent();
+	UMeshPaintModeHelpers::PasteVertexColors(StaticMeshComponents, CopiedColorsByComponent);
+	UpdateCachedVertexDataSize();
+}
+
+void UMeshPaintMode::FixVertexColors()
+{
+	FScopedTransaction Transaction(LOCTEXT("LevelMeshPainter_TransactionFixInstColors", "Fixing Per-Instance Vertex Colors"));
+	const TArray<UStaticMeshComponent*> StaticMeshComponents = GetSelectedComponents<UStaticMeshComponent>();
+	for (UStaticMeshComponent* Component : StaticMeshComponents)
+	{
+		Component->FixupOverrideColorsIfNecessary();
+	}
+}
+
+bool UMeshPaintMode::DoesRequireVertexColorsFixup() const
+{
+	const TArray<UStaticMeshComponent*> StaticMeshComponents = GetSelectedComponents<UStaticMeshComponent>();
+	bool bAnyMeshNeedsFixing = false;
+	/** Check if there are any static mesh components which require fixing */
+	for (UStaticMeshComponent* Component : StaticMeshComponents)
+	{
+		bAnyMeshNeedsFixing |= Component->RequiresOverrideVertexColorsFixup();
+	}
+
+	return bAnyMeshNeedsFixing;
+}
+
+void UMeshPaintMode::RemoveVertexColors()
+{
+	FScopedTransaction Transaction(LOCTEXT("LevelMeshPainter_TransactionRemoveInstColors", "Removing Per-Instance Vertex Colors"));
+	const TArray<UStaticMeshComponent*> StaticMeshComponents = GetSelectedComponents<UStaticMeshComponent>();
+	for (UStaticMeshComponent* Component : StaticMeshComponents)
+	{
+		UMeshPaintingToolset::RemoveComponentInstanceVertexColors(Component);
+	}
+
+	UpdateCachedVertexDataSize();
+}
+
+
+void UMeshPaintMode::PropagateVertexColorsToLODs()
+{
+	//Only show the lost data warning if there is actually some data to lose
+	bool bAbortChange = false;
+	UMeshToolManager* MeshToolManager = Cast<UMeshToolManager>(GetToolManager());
+	TArray<UMeshComponent*> PaintableComponents = MeshToolManager->GetPaintableMeshComponents();
+	const bool bSelectionContainsPerLODColors = MeshToolManager->SelectionContainsPerLODColors();
+	if (bSelectionContainsPerLODColors)
+	{
+		//Warn the user they will lose custom painting data
+		FSuppressableWarningDialog::FSetupInfo SetupInfo(LOCTEXT("LooseLowersLODsVertexColorsPrompt_Message", "Propagating Vertex Colors from LOD0 to all lower LODs. This mean all lower LODs custom vertex painting will be lost."),
+			LOCTEXT("LooseLowersLODsVertexColorsPrompt_Title", "Warning: Lowers LODs custom vertex painting will be lost!"), "Warning_LooseLowersLODsVertexColorsPrompt");
+
+		SetupInfo.ConfirmText = LOCTEXT("LooseLowersLODsVertexColorsPrompt_ConfirmText", "Continue");
+		SetupInfo.CancelText = LOCTEXT("LooseLowersLODsVertexColorsPrompt_CancelText", "Abort");
+		SetupInfo.CheckBoxText = LOCTEXT("LooseLowersLODsVertexColorsPrompt_CheckBoxText", "Always copy vertex colors without prompting");
+
+		FSuppressableWarningDialog LooseLowersLODsVertexColorsWarning(SetupInfo);
+	
+		// Prompt the user to see if they really want to propagate the base lod vert colors to the lowers LODs.
+		if (LooseLowersLODsVertexColorsWarning.ShowModal() == FSuppressableWarningDialog::Cancel)
+		{
+			bAbortChange = true;
+		}
+		else
+		{
+			// Reset the state flag as we'll be removing all per-lod colors 
+			MeshToolManager->ClearSelectionLODColors();
+			UMeshPaintModeHelpers::RemovePerLODColors(PaintableComponents);
+		}
+	}
+
+	//The user cancel the change, avoid changing the value
+	if (bAbortChange)
+	{
+		return;
+	}
+
+	for (UMeshComponent* SelectedComponent : PaintableComponents)
+	{
+		if (SelectedComponent)
+		{
+			TSharedPtr<IMeshPaintComponentAdapter> MeshAdapter = Cast<UMeshToolManager>(GetToolManager())->GetAdapterForComponent(SelectedComponent);
+			UMeshPaintingToolset::ApplyVertexColorsToAllLODs(*MeshAdapter, SelectedComponent);
+			FComponentReregisterContext ReregisterContext(SelectedComponent);
+		}
+	}
+
+	MeshToolManager->Refresh();
+}
+
+template<typename ComponentClass>
+TArray<ComponentClass*> UMeshPaintMode::GetSelectedComponents() const
+{
+	TArray<ComponentClass*> Components;
+	FToolBuilderState SelectionState;
+	GetToolManager()->GetContextQueriesAPI()->GetCurrentSelectionState(SelectionState);
+//	if (PaintSettings->PaintMode == EPaintMode::Textures)
+	{
+
+		for (int32 SelectionIndex = 0; SelectionIndex < SelectionState.SelectedComponents.Num(); ++SelectionIndex)
+		{
+			ComponentClass* SelectedComponent = Cast<ComponentClass>(SelectionState.SelectedComponents[SelectionIndex]);
+			if (SelectedComponent)
+			{
+				Components.AddUnique(SelectedComponent);
+			}
+		}
+	}
+
+	if (Components.Num() == 0)
+	{
+		for (int32 SelectionIndex = 0; SelectionIndex < SelectionState.SelectedActors.Num(); ++SelectionIndex)
+		{
+			AActor* SelectedActor = Cast<AActor>(SelectionState.SelectedActors[SelectionIndex]);
+			if (SelectedActor)
+			{
+				TInlineComponentArray<ComponentClass*> ActorComponents;
+				SelectedActor->GetComponents(ActorComponents);
+				for (ComponentClass* Component : ActorComponents)
+				{
+					Components.AddUnique(Component);
+				}
+			}
+		}
+	}
+
+	return Components;
+}
+
+template TArray<UStaticMeshComponent*> UMeshPaintMode::GetSelectedComponents<UStaticMeshComponent>() const;
+template TArray<USkeletalMeshComponent*> UMeshPaintMode::GetSelectedComponents<USkeletalMeshComponent>() const;
+template TArray<UMeshComponent*> UMeshPaintMode::GetSelectedComponents<UMeshComponent>() const;
+
+
+
+
+void UMeshPaintMode::UpdateCachedVertexDataSize()
+{
+	CachedVertexDataSize = 0;
+
+	const bool bInstance = true;
+	if (UMeshToolManager* MeshToolManager = Cast<UMeshToolManager>(GetToolManager()))
+	{
+		for (UMeshComponent* SelectedComponent : MeshToolManager->GetPaintableMeshComponents())
+		{
+			int32 NumLODs = UMeshPaintingToolset::GetNumberOfLODs(SelectedComponent);
+			for (int32 LODIndex = 0; LODIndex < NumLODs; ++LODIndex)
+			{
+				CachedVertexDataSize += UMeshPaintingToolset::GetVertexColorBufferSize(SelectedComponent, LODIndex, bInstance);
+			}
+		}
+	}
+}
+
+
+
+void UMeshPaintMode::ActivateDefaultTool()
+{
+	if (CurrentPaletteName == UMeshPaintMode::MeshPaintMode_Color)
+	{
+		ToolsContext->StartTool(EToolSide::Mouse, TEXT("MeshClickTool"));
+	}
+	if (CurrentPaletteName == UMeshPaintMode::MeshPaintMode_Weights)
+	{
+		ToolsContext->StartTool(EToolSide::Mouse, TEXT("MeshClickTool"));
+	}
+	if (CurrentPaletteName == UMeshPaintMode::MeshPaintMode_Weights)
+	{
+
+	}
+}
+
+void UMeshPaintMode::UpdateOnPaletteChange()
+{
+	UpdateSelectedMeshes();
+}
+
+void UMeshPaintMode::OnResetViewMode()
+{
+	// Reset viewport color mode for all active viewports
+	for (FEditorViewportClient* ViewportClient : GEditor->GetAllViewportClients())
+	{
+		if (!ViewportClient || ViewportClient->GetModeTools() != GetModeManager())
+		{
+			continue;
+		}
+
+		UMeshPaintModeHelpers::SetViewportColorMode(EMeshPaintDataColorViewMode::Normal, ViewportClient);
+	}
+}
+
+
+#undef LOCTEXT_NAMESPACE
+
diff --git a/Engine/Plugins/MeshPainting/Source/MeshPaintEditorMode/Private/MeshPaintMode.h b/Engine/Plugins/MeshPainting/Source/MeshPaintEditorMode/Private/MeshPaintMode.h
new file mode 100644
--- /dev/null
+++ b/Engine/Plugins/MeshPainting/Source/MeshPaintEditorMode/Private/MeshPaintMode.h
@@ -0,0 +1,101 @@
+// Copyright 1998-2019 Epic Games, Inc. All Rights Reserved.
+
+#pragma once
+
+#include "Tools/UEdMode.h"
+#include "Framework/Commands/UIAction.h"
+#include "Framework/Commands/UICommandList.h"
+#include "Containers/ArrayView.h"
+#include "LevelEditorViewport.h"	// For FTrackingTransaction
+#include "UnrealEdMisc.h"
+#include "Framework/Commands/UIAction.h"
+#include "Framework/Commands/UICommandInfo.h"
+#include "Framework/Commands/UICommandList.h"
+#include "IMeshPaintComponentAdapter.h"
+#include "MeshPaintHelpers.h"
+#include "MeshVertexPaintingTool.h"
+#include "MeshPaintMode.generated.h"
+
+class UMeshVertexPaintingToolProperties;
+class UMeshColorPaintingToolProperties;
+class UMeshPaintModeSettings;
+class IMeshPaintComponentAdapter;
+class UMeshComponent;
+
+/**
+ * Mesh paint Mode.  Extends editor viewports with the ability to paint data on meshes
+ */
+UCLASS()
+class UMeshPaintMode : public UEdMode
+{
+public:
+	GENERATED_BODY()
+
+	/** Default constructor for UMeshPaintMode */
+	UMeshPaintMode();
+	static UMeshVertexPaintingToolProperties* GetVertexToolProperties();
+	static UMeshColorPaintingToolProperties* GetColorToolProperties();
+	static UMeshWeightPaintingToolProperties* GetWeightToolProperties();
+	static UMeshPaintMode* GetMeshPaintMode();
+	virtual void Enter() override;
+	virtual void Exit() override;
+	virtual void CreateToolkit() override;
+	virtual void Tick(FEditorViewportClient* ViewportClient, float DeltaTime) override;
+	static FName MeshPaintMode_Color;
+	static FName MeshPaintMode_Texture;
+	static FName MeshPaintMode_Weights;
+
+	virtual TMap<FName, TArray<TSharedPtr<FUICommandInfo>>> GetModeCommands() const override;
+	/** Returns the instance of ComponentClass found in the current Editor selection */
+	template<typename ComponentClass>
+	TArray<ComponentClass*> GetSelectedComponents() const;
+
+	uint32 GetCachedVertexDataSize() const
+	{
+		return CachedVertexDataSize;
+	}
+
+protected:
+
+	/** Binds UI commands to actions for the mesh paint mode */
+	virtual void BindCommands();
+
+	// UEdMode interface
+	virtual void OnToolStarted(UInteractiveToolManager* Manager, UInteractiveTool* Tool) override;
+	virtual void OnToolEnded(UInteractiveToolManager* Manager, UInteractiveTool* Tool) override;
+	virtual bool InputAxis(FEditorViewportClient* InViewportClient, FViewport* Viewport, int32 ControllerId, FKey Key, float Delta, float DeltaTime) override;
+	virtual void ActorSelectionChangeNotify() override;
+	virtual void ActivateDefaultTool() override;
+	virtual void UpdateOnPaletteChange() override;
+	// end UEdMode Interface
+	void UpdateSelectedMeshes();
+	void OnResetViewMode();
+	void FillWithVertexColor();
+	void PropagateVertexColorsToAsset();
+	bool CanPropagateVertexColors() const;
+	void ImportVertexColors();
+	void SavePaintedAssets();
+	bool CanSaveMeshPackages() const;
+	bool CanRemoveInstanceColors() const;
+	bool CanPasteInstanceVertexColors() const;
+	bool CanCopyInstanceVertexColors() const;
+	bool CanPropagateVertexColorsToLODs() const;
+	/** Copy and pasting functionality from and to a StaticMeshComponent (currently only supports per instance)*/
+	void CopyVertexColors();
+	void PasteVertexColors();
+	void FixVertexColors();
+	bool DoesRequireVertexColorsFixup() const;
+	void RemoveVertexColors();
+	void PropagateVertexColorsToLODs();
+	void UpdateCachedVertexDataSize();
+
+protected:
+	UPROPERTY(Transient)
+	UMeshPaintModeSettings* ModeSettings;
+	// End vertex paint state
+	FGetSelectedMeshComponents MeshComponentDelegate;
+	uint32 CachedVertexDataSize;
+	bool bRecacheVertexDataSize;
+
+};
+
diff --git a/Engine/Plugins/MeshPainting/Source/MeshPaintEditorMode/Private/MeshPaintModeCommands.cpp b/Engine/Plugins/MeshPainting/Source/MeshPaintEditorMode/Private/MeshPaintModeCommands.cpp
new file mode 100644
--- /dev/null
+++ b/Engine/Plugins/MeshPainting/Source/MeshPaintEditorMode/Private/MeshPaintModeCommands.cpp
@@ -0,0 +1,112 @@
+// Copyright 1998-2019 Epic Games, Inc. All Rights Reserved.
+
+#include "MeshPaintModeCommands.h"
+#include "Framework/Commands/UIAction.h"
+#include "MeshPaintMode.h"
+#include "SingleSelectionTool.h"
+#include "MeshVertexPaintingTool.h"
+
+#define LOCTEXT_NAMESPACE "MeshPaintEditorModeCommands"
+
+
+
+void FMeshPaintingToolActionCommands::GetToolDefaultObjectList(TArray<UInteractiveTool*>& ToolCDOs)
+{
+	ToolCDOs.Add(GetMutableDefault<USingleSelectionTool>());
+	ToolCDOs.Add(GetMutableDefault<UMeshColorPaintingTool>());
+}
+
+
+
+
+void FMeshPaintingToolActionCommands::RegisterAllToolActions()
+{
+	FMeshPaintingToolActionCommands::Register();
+}
+
+void FMeshPaintingToolActionCommands::UnregisterAllToolActions()
+{
+	FMeshPaintingToolActionCommands::Unregister();
+}
+
+void FMeshPaintingToolActionCommands::UpdateToolCommandBinding(UInteractiveTool* Tool, TSharedPtr<FUICommandList> UICommandList, bool bUnbind)
+{
+	if (FMeshPaintingToolActionCommands::IsRegistered())
+	{
+		!bUnbind ? FMeshPaintingToolActionCommands::Get().BindCommandsForCurrentTool(UICommandList, Tool) : FMeshPaintingToolActionCommands::Get().UnbindActiveCommands(UICommandList);
+	}
+}
+
+
+
+
+void FMeshPaintEditorModeCommands::RegisterCommands()
+{
+	TArray<TSharedPtr<FUICommandInfo>> ColorCommands;
+	TArray<TSharedPtr<FUICommandInfo>> WeightCommands;
+	TArray<TSharedPtr<FUICommandInfo>> VertexCommands;
+	 
+	UI_COMMAND(Select, "Select", "Select the mesh", EUserInterfaceActionType::ToggleButton, FInputChord());
+	ColorCommands.Add(Select);
+	WeightCommands.Add(Select);
+	UI_COMMAND(ColorPaint, "Paint", "Paint the mesh", EUserInterfaceActionType::ToggleButton, FInputChord());
+	ColorCommands.Add(ColorPaint);
+	UI_COMMAND(WeightPaint, "Paint", "Paint the mesh", EUserInterfaceActionType::ToggleButton, FInputChord());
+	WeightCommands.Add(WeightPaint);
+	UI_COMMAND(Fill, "Fill", "Fills the selected Meshes with the Paint Color", EUserInterfaceActionType::Button, FInputChord());
+	VertexCommands.Add(Fill);
+	UI_COMMAND(Propagate, "Apply", "Propagates Instance Vertex Colors to the Source Meshes", EUserInterfaceActionType::Button, FInputChord());
+	VertexCommands.Add(Propagate);
+ 	UI_COMMAND(Import, "Import", "Imports Vertex Colors from a TGA Texture File to the Selected Meshes", EUserInterfaceActionType::Button, FInputChord());
+ 	VertexCommands.Add(Import);
+ 	UI_COMMAND(Save, "Save", "Saves the Source Meshes for the selected Mesh Components", EUserInterfaceActionType::Button, FInputChord());
+	VertexCommands.Add(Save);
+	UI_COMMAND(Copy, "Copy", "Copies Vertex Colors from the selected Mesh Components", EUserInterfaceActionType::Button, FInputChord());
+ 	VertexCommands.Add(Copy);
+	UI_COMMAND(Paste, "Paste", "Tried to Paste Vertex Colors on the selected Mesh Components", EUserInterfaceActionType::Button, FInputChord());
+ 	VertexCommands.Add(Paste);
+	UI_COMMAND(Remove, "Remove", "Removes Vertex Colors from the selected Mesh Components", EUserInterfaceActionType::Button, FInputChord());
+	VertexCommands.Add(Remove);
+	UI_COMMAND(Fix, "Fix", "If necessary fixes Vertex Colors applied to the selected Mesh Components", EUserInterfaceActionType::Button, FInputChord());
+	VertexCommands.Add(Fix);
+	UI_COMMAND(PropagateVertexColorsToLODs, "All LODs", "Applied the Vertex Colors from LOD0 to all LOD levels", EUserInterfaceActionType::Button, FInputChord());
+	VertexCommands.Add(PropagateVertexColorsToLODs);
+	ColorCommands.Append(VertexCommands);
+	WeightCommands.Append(VertexCommands);
+	Commands.Add(UMeshPaintMode::MeshPaintMode_Color, ColorCommands);
+	Commands.Add(UMeshPaintMode::MeshPaintMode_Weights, WeightCommands);
+
+	TArray<TSharedPtr<FUICommandInfo>> TextureCommands;
+	UI_COMMAND(PropagateTexturePaint, "Apply", "Propagates Modifications to the Textures", EUserInterfaceActionType::Button, FInputChord());
+	TextureCommands.Add(PropagateTexturePaint);
+	UI_COMMAND(SaveTexturePaint, "Save", "Saves the Modified Textures for the selected Mesh Components", EUserInterfaceActionType::Button, FInputChord());
+	TextureCommands.Add(SaveTexturePaint);
+	Commands.Add(UMeshPaintMode::MeshPaintMode_Texture, TextureCommands);
+
+	// // 	UI_COMMAND(NextTexture, "Next Texture", "Cycle To Next Texture", EUserInterfaceActionType::Button, FInputChord(EKeys::Period));
+	// // 	Commands.Add(NextTexture);
+	// // 	UI_COMMAND(PreviousTexture, "Previous Texture", "Cycle To Previous Texture", EUserInterfaceActionType::Button, FInputChord(EKeys::Comma));
+	// // 	Commands.Add(PreviousTexture);
+	// // 
+	// // 	UI_COMMAND(CommitTexturePainting, "Commit Texture Painting", "Commits Texture Painting Changes", EUserInterfaceActionType::Button, FInputChord(EKeys::C, EModifierKey::Control | EModifierKey::Shift));
+	// // 	Commands.Add(CommitTexturePainting);
+	// // 
+	// 
+	// 
+	// 
+	// // 
+	// // 	UI_COMMAND(SwitchForeAndBackgroundColor, "Switch Fore and Background Color", "Switches the Fore and Background Colors used for Vertex Painting", EUserInterfaceActionType::None, FInputChord(EKeys::X));
+	// // 	Commands.Add(SwitchForeAndBackgroundColor);
+	// // 
+	// 
+	// // 
+	// // 
+	// // 	UI_COMMAND(CycleToNextLOD, "Cycle to next Mesh LOD", "Cycles to the next possible Mesh LOD to Paint on", EUserInterfaceActionType::None, FInputChord(EKeys::N));
+	// // 	Commands.Add(CycleToNextLOD);
+	// // 
+	// // 	UI_COMMAND(CycleToPreviousLOD, "Cycle to previous Mesh LOD", "Cycles to the previous possible Mesh LOD to Paint on", EUserInterfaceActionType::None, FInputChord(EKeys::B));
+	// // 	Commands.Add(CycleToPreviousLOD);
+}
+
+#undef LOCTEXT_NAMESPACE
+
diff --git a/Engine/Plugins/MeshPainting/Source/MeshPaintEditorMode/Private/MeshPaintModeSettings.cpp b/Engine/Plugins/MeshPainting/Source/MeshPaintEditorMode/Private/MeshPaintModeSettings.cpp
new file mode 100644
--- /dev/null
+++ b/Engine/Plugins/MeshPainting/Source/MeshPaintEditorMode/Private/MeshPaintModeSettings.cpp
@@ -0,0 +1,5 @@
+// Copyright 1998-2019 Epic Games, Inc. All Rights Reserved.
+
+#include "MeshPaintModeSettings.h"
+#include "UnrealEdMisc.h"
+
diff --git a/Engine/Plugins/MeshPainting/Source/MeshPaintEditorMode/Private/MeshPaintModeToolkit.cpp b/Engine/Plugins/MeshPainting/Source/MeshPaintEditorMode/Private/MeshPaintModeToolkit.cpp
new file mode 100644
--- /dev/null
+++ b/Engine/Plugins/MeshPainting/Source/MeshPaintEditorMode/Private/MeshPaintModeToolkit.cpp
@@ -0,0 +1,51 @@
+// Copyright 1998-2019 Epic Games, Inc. All Rights Reserved.
+
+#include "MeshPaintModeToolkit.h"
+#include "MeshPaintMode.h"
+
+#define LOCTEXT_NAMESPACE "MeshPaintModeToolkit"
+
+
+FName FMeshPaintModeToolkit::GetToolkitFName() const
+{
+	return FName( "MeshPaintMode" );
+}
+
+
+FText FMeshPaintModeToolkit::GetBaseToolkitName() const
+{
+	return LOCTEXT( "ToolkitName", "Mesh Paint Mode" );
+}
+
+
+class UEdMode* FMeshPaintModeToolkit::GetScriptableEditorMode() const
+{
+	return GLevelEditorModeTools().GetActiveScriptableMode( "MeshPaintMode" );
+}
+
+void FMeshPaintModeToolkit::GetToolPaletteNames(TArray<FName>& PaletteNames) const
+{
+	PaletteNames.Add(UMeshPaintMode::MeshPaintMode_Color);
+	PaletteNames.Add(UMeshPaintMode::MeshPaintMode_Weights);
+	PaletteNames.Add(UMeshPaintMode::MeshPaintMode_Texture);
+}
+
+
+FText FMeshPaintModeToolkit::GetToolPaletteDisplayName(FName Palette) const
+{
+	if (Palette == UMeshPaintMode::MeshPaintMode_Color)
+	{
+		return LOCTEXT("MeshPaintMode_Color", "Colors");
+	}
+	if (Palette == UMeshPaintMode::MeshPaintMode_Weights)
+	{
+		return LOCTEXT("MeshPaintMode_Weights", "Weights");
+	}
+	if (Palette == UMeshPaintMode::MeshPaintMode_Texture)
+	{
+		return LOCTEXT("MeshPaintMode_Texture", "Textures");
+	}
+	return FText();
+}
+
+#undef LOCTEXT_NAMESPACE
diff --git a/Engine/Plugins/MeshPainting/Source/MeshPaintEditorMode/Public/MeshPaintEditorMode.h b/Engine/Plugins/MeshPainting/Source/MeshPaintEditorMode/Public/MeshPaintEditorMode.h
new file mode 100644
--- /dev/null
+++ b/Engine/Plugins/MeshPainting/Source/MeshPaintEditorMode/Public/MeshPaintEditorMode.h
@@ -0,0 +1,5 @@
+// Copyright 1998-2019 Epic Games, Inc. All Rights Reserved.
+
+#pragma once
+
+#include "CoreMinimal.h"
diff --git a/Engine/Plugins/MeshPainting/Source/MeshPaintEditorMode/Public/MeshPaintModeCommands.h b/Engine/Plugins/MeshPainting/Source/MeshPaintEditorMode/Public/MeshPaintModeCommands.h
new file mode 100644
--- /dev/null
+++ b/Engine/Plugins/MeshPainting/Source/MeshPaintEditorMode/Public/MeshPaintModeCommands.h
@@ -0,0 +1,101 @@
+// Copyright 1998-2019 Epic Games, Inc. All Rights Reserved.
+
+#pragma once
+
+#include "Framework/Commands/UICommandInfo.h"
+#include "Framework/Commands/Commands.h"
+#include "EditorStyleSet.h"
+#include "Tools/InteractiveToolsCommands.h"
+
+
+class FMeshPaintEditorModeCommands : public TCommands<FMeshPaintEditorModeCommands>
+{
+public:
+	FMeshPaintEditorModeCommands() : TCommands<FMeshPaintEditorModeCommands>
+		(
+			"MeshPaint",
+			NSLOCTEXT("MeshPaintEditorMode", "MeshPaintingModeCommands", "Mesh Painting Mode"),
+			NAME_None,
+			FEditorStyle::GetStyleSetName()
+			)
+	{}
+
+	/**
+	* Initialize commands
+	*/
+	virtual void RegisterCommands() override;
+
+	static TMap<FName, TArray<TSharedPtr<FUICommandInfo>>> GetCommands()
+	{
+		return FMeshPaintEditorModeCommands::Get().Commands;
+	}
+
+public:
+	TSharedPtr<FUICommandInfo> NextTexture;
+	TSharedPtr<FUICommandInfo> PreviousTexture;
+	TSharedPtr<FUICommandInfo> CommitTexturePainting;
+
+	TSharedPtr<FUICommandInfo> Copy;
+	TSharedPtr<FUICommandInfo> Paste;
+	TSharedPtr<FUICommandInfo> Remove;
+	TSharedPtr<FUICommandInfo> Fix;
+	TSharedPtr<FUICommandInfo> Fill;
+	TSharedPtr<FUICommandInfo> Propagate;
+	TSharedPtr<FUICommandInfo> Import;
+	TSharedPtr<FUICommandInfo> Save;
+
+	TSharedPtr<FUICommandInfo> SwitchForeAndBackgroundColor;
+	TSharedPtr<FUICommandInfo> CycleToNextLOD;
+	TSharedPtr<FUICommandInfo> CycleToPreviousLOD;
+
+	TSharedPtr<FUICommandInfo> PropagateTexturePaint;
+	TSharedPtr<FUICommandInfo> SaveTexturePaint;
+
+	TSharedPtr<FUICommandInfo> PropagateVertexColorsToLODs;
+
+	TSharedPtr<FUICommandInfo> ColorPaint;
+	TSharedPtr<FUICommandInfo> WeightPaint;
+	TSharedPtr<FUICommandInfo> TexturePaint;
+	TSharedPtr<FUICommandInfo> Select;
+
+	TMap<FName, TArray<TSharedPtr<FUICommandInfo>>> Commands;
+};
+
+
+class FMeshPaintingToolActionCommands : public TInteractiveToolCommands<FMeshPaintingToolActionCommands>
+{
+public:
+	FMeshPaintingToolActionCommands() :
+		TInteractiveToolCommands<FMeshPaintingToolActionCommands>(
+			"MeshPaintingTools", // Context name for fast lookup
+			NSLOCTEXT("MeshPaintEditorMode", "MeshPaintingToolsCommands", "Mesh Painting Tools"), // Localized context name for displaying
+			NAME_None, // Parent
+			FEditorStyle::GetStyleSetName() // Icon Style Set
+			)
+	{
+	}
+
+	virtual void GetToolDefaultObjectList(TArray<UInteractiveTool*>& ToolCDOs) override;
+
+
+	/**
+	 * interface that hides various per-tool action sets
+	 */
+
+	 /**
+	  * Register all Tool command sets. Call this in module startup
+	  */
+	static void RegisterAllToolActions();
+
+	/**
+	 * Unregister all Tool command sets. Call this from module shutdown.
+	 */
+	static void UnregisterAllToolActions();
+
+	/**
+	 * Add or remove commands relevant to Tool to the given UICommandList.
+	 * Call this when the active tool changes (eg on ToolManager.OnToolStarted / OnToolEnded)
+	 * @param bUnbind if true, commands are removed, otherwise added
+	 */
+	static void UpdateToolCommandBinding(UInteractiveTool* Tool, TSharedPtr<FUICommandList> UICommandList, bool bUnbind = false);
+};
\ No newline at end of file
diff --git a/Engine/Plugins/MeshPainting/Source/MeshPaintEditorMode/Public/MeshPaintModeHelpers.h b/Engine/Plugins/MeshPainting/Source/MeshPaintEditorMode/Public/MeshPaintModeHelpers.h
new file mode 100644
--- /dev/null
+++ b/Engine/Plugins/MeshPainting/Source/MeshPaintEditorMode/Public/MeshPaintModeHelpers.h
@@ -0,0 +1,72 @@
+// Copyright 1998-2019 Epic Games, Inc. All Rights Reserved.
+
+#pragma once
+
+#include "CoreMinimal.h"
+#include "Engine/EngineTypes.h"
+#include "MeshPaintingToolsetTypes.h"
+#include "Kismet/BlueprintFunctionLibrary.h"
+#include "IMeshPaintComponentAdapter.h"
+
+class FMeshPaintParameters;
+class UImportVertexColorOptions;
+class UTexture2D;
+class UStaticMeshComponent;
+class UStaticMesh;
+class USkeletalMesh;
+class IMeshPaintComponentAdapter;
+class UPaintBrushSettings;
+class FEditorViewportClient;
+class UMeshComponent;
+class USkeletalMeshComponent;
+class UViewportInteractor;
+class FViewport;
+class FPrimitiveDrawInterface;
+class FSceneView;
+
+struct FStaticMeshComponentLODInfo;
+struct FPerComponentVertexColorData;
+
+enum class EMeshPaintDataColorViewMode : uint8;
+
+class MESHPAINTEDITORMODE_API UMeshPaintModeHelpers : public UBlueprintFunctionLibrary
+{
+public:
+	/** Forces the Viewport Client to render using the given Viewport Color ViewMode */
+	static void SetViewportColorMode(EMeshPaintDataColorViewMode ColorViewMode, FEditorViewportClient* ViewportClient);
+
+	/** Sets whether or not the level viewport should be real time rendered move or viewport as parameter? */
+	static void SetRealtimeViewport(bool bRealtime);
+
+
+	/** Helper function to import Vertex Colors from a Texture to the specified MeshComponent (makes use of SImportVertexColorsOptions Widget) */
+	static void ImportVertexColorsFromTexture(UMeshComponent* MeshComponent);
+
+	/** Imports vertex colors from a Texture to the specified Skeletal Mesh according to user-set options */
+	static void ImportVertexColorsToSkeletalMesh(USkeletalMesh* SkeletalMesh, const UImportVertexColorOptions* Options, UTexture2D* Texture);
+
+	struct FPaintRay
+	{
+		FVector CameraLocation;
+		FVector RayStart;
+		FVector RayDirection;
+		UViewportInteractor* ViewportInteractor;
+	};
+
+
+	static bool RetrieveViewportPaintRays(const FSceneView* View, FViewport* Viewport, FPrimitiveDrawInterface* PDI, TArray<FPaintRay>& OutPaintRays);
+
+	/** Imports vertex colors from a Texture to the specified Static Mesh according to user-set options */
+	static void ImportVertexColorsToStaticMesh(UStaticMesh* StaticMesh, const UImportVertexColorOptions* Options, UTexture2D* Texture);
+
+	/** Imports vertex colors from a Texture to the specified Static Mesh Component according to user-set options */
+	static void ImportVertexColorsToStaticMeshComponent(UStaticMeshComponent* StaticMeshComponent, const UImportVertexColorOptions* Options, UTexture2D* Texture);
+
+	static void PropagateVertexColors(const TArray<UStaticMeshComponent *> StaticMeshComponents);
+	static bool CanPropagateVertexColors(TArray<UStaticMeshComponent*>& StaticMeshComponents, TArray<UStaticMesh*>& StaticMeshes, int32 NumInstanceVertexColorBytes);
+	static void CopyVertexColors(const TArray<UStaticMeshComponent*> StaticMeshComponents, TArray<FPerComponentVertexColorData>& CopiedVertexColors);
+	static bool CanCopyInstanceVertexColors(const TArray<UStaticMeshComponent*>& StaticMeshComponents, int32 PaintingMeshLODIndex);
+	static void PasteVertexColors(const TArray<UStaticMeshComponent*>& StaticMeshComponents, TArray<FPerComponentVertexColorData>& CopiedColorsByComponent);
+	static bool CanPasteInstanceVertexColors(const TArray<UStaticMeshComponent*>& StaticMeshComponents, const TArray<FPerComponentVertexColorData>& CopiedColorsByComponent);
+	static void RemovePerLODColors(const TArray<UMeshComponent*>& PaintableComponents);
+};
\ No newline at end of file
diff --git a/Engine/Plugins/MeshPainting/Source/MeshPaintingToolset/Private/MeshSelect.cpp b/Engine/Plugins/MeshPainting/Source/MeshPaintingToolset/Private/MeshSelect.cpp
new file mode 100644
--- /dev/null
+++ b/Engine/Plugins/MeshPainting/Source/MeshPaintingToolset/Private/MeshSelect.cpp
@@ -0,0 +1,121 @@
+// Copyright 1998-2019 Epic Games, Inc. All Rights Reserved.
+
+
+#include "MeshSelect.h"
+#include "InputState.h"
+#include "InteractiveToolManager.h"
+#include "ToolContextInterfaces.h"
+#include "MeshPaintHelpers.h"
+#include "Components/MeshComponent.h"
+#include "IMeshPaintComponentAdapter.h"
+#include "MeshPaintAdapterFactory.h"
+#include "EngineUtils.h"
+#include "Editor.h"
+
+#define LOCTEXT_NAMESPACE "MeshSelection"
+
+
+bool UMeshClickToolBuilder::CanBuildTool(const FToolBuilderState& SceneState) const
+{
+	return true;
+}
+
+UInteractiveTool* UMeshClickToolBuilder::BuildTool(const FToolBuilderState& SceneState) const
+{
+	UMeshClickTool* NewTool = NewObject<UMeshClickTool>(SceneState.ToolManager);
+
+	return NewTool;
+}
+
+
+
+
+UMeshClickTool::UMeshClickTool()
+{
+}
+
+void UMeshClickTool::Setup()
+{
+	UInteractiveTool::Setup();
+
+	// add default button input behaviors for devices
+	USingleClickInputBehavior* MouseBehavior = NewObject<USingleClickInputBehavior>();
+	MouseBehavior->Initialize(this);
+	MouseBehavior->Modifiers.RegisterModifier(AdditiveSelectionModifier, FInputDeviceState::IsShiftKeyDown);
+	AddInputBehavior(MouseBehavior);
+}
+
+void UMeshClickTool::OnUpdateModifierState(int ModifierID, bool bIsOn)
+{
+	if (ModifierID == AdditiveSelectionModifier)
+	{
+		bAddToSelectionSet = bIsOn;
+	}
+}
+
+FInputRayHit UMeshClickTool::IsHitByClick(const FInputDeviceRay& ClickPos)
+{
+	return FInputRayHit(0.0f);
+}
+
+void UMeshClickTool::OnClicked(const FInputDeviceRay& ClickPos)
+{
+	if (UMeshToolManager* MeshToolManager = Cast<UMeshToolManager>(GetToolManager()))
+	{
+		const FVector& RayOrigin = ClickPos.WorldRay.Origin;
+		const FVector& RayDirection = ClickPos.WorldRay.Direction;
+		FRay Ray = FRay(RayOrigin, RayDirection);
+		const FVector TraceStart(RayOrigin);
+		const FVector TraceEnd(RayOrigin + RayDirection * HALF_WORLD_MAX);
+		MeshToolManager->ResetState();
+		//Default iterator only iterates over active levels.
+		const EActorIteratorFlags Flags = EActorIteratorFlags::SkipPendingKill;
+		// TODO: The tool needs to know the world its acting on
+		for (TActorIterator<AActor> It(GEditor->GetEditorWorldContext(false).World(), AActor::StaticClass(), Flags); It; ++It)
+		{
+			AActor* Actor = *It;
+			if (Actor->IsEditable() &&
+				Actor->IsListedInSceneOutliner() &&					// Only add actors that are allowed to be selected and drawn in editor
+				!Actor->IsTemplate() &&								// Should never happen, but we never want CDOs
+				!Actor->HasAnyFlags(RF_Transient))				// Don't add transient actors in non-play worlds
+			{
+				TArray<UActorComponent*> CandidateComponents = Actor->K2_GetComponentsByClass(UMeshComponent::StaticClass());
+				for (UActorComponent* CandidateComponent : CandidateComponents)
+				{
+					UMeshComponent* MeshComponent = Cast<UMeshComponent>(CandidateComponent);
+					TSharedPtr<IMeshPaintComponentAdapter> MeshAdapter = FMeshPaintComponentAdapterFactory::CreateAdapterForMesh(MeshComponent, 0);
+					MeshToolManager->AddPaintableMeshComponent(MeshComponent);
+					MeshToolManager->AddToComponentToAdapterMap(MeshComponent, MeshAdapter);
+				}
+					
+			}
+		}
+		FHitResult TraceHitResult(1.0f);
+		MeshToolManager->FindHitResult(Ray, TraceHitResult);
+		MeshToolManager->ResetState();
+		TArray<AActor*> SelectedActors;
+		if (UMeshComponent* MeshComponent = Cast<UMeshComponent>(TraceHitResult.GetComponent()))
+		{
+			TSharedPtr<IMeshPaintComponentAdapter> MeshAdapter = FMeshPaintComponentAdapterFactory::CreateAdapterForMesh(MeshComponent, 0);
+			if (MeshComponent->IsVisible() && MeshAdapter.IsValid() && MeshAdapter->IsValid())
+			{
+				MeshToolManager->AddPaintableMeshComponent(MeshComponent);
+				MeshToolManager->AddToComponentToAdapterMap(MeshComponent, MeshAdapter);
+				SelectedActors.Add(Cast<AActor>(MeshComponent->GetOuter()));
+				MeshAdapter->OnAdded();
+			}
+		}
+
+		GetToolManager()->BeginUndoTransaction(LOCTEXT("MeshSelection", "Select Mesh"));
+
+
+		FSelectedOjectsChangeList NewSelection;
+		// TODO add CTRL handling
+		NewSelection.ModificationType = bAddToSelectionSet ? ESelectedObjectsModificationType::Add : ESelectedObjectsModificationType::Replace;
+		NewSelection.Actors.Append(SelectedActors);
+		GetToolManager()->RequestSelectionChange(NewSelection);
+		GetToolManager()->EndUndoTransaction();
+	}
+}
+
+#undef LOCTEXT_NAMESPACE
diff --git a/Engine/Plugins/MeshPainting/Source/MeshPaintingToolset/Public/MeshPaintAdapterFactory.h b/Engine/Plugins/MeshPainting/Source/MeshPaintingToolset/Public/MeshPaintAdapterFactory.h
new file mode 100644
--- /dev/null
+++ b/Engine/Plugins/MeshPainting/Source/MeshPaintingToolset/Public/MeshPaintAdapterFactory.h
@@ -0,0 +1,21 @@
+// Copyright 1998-2019 Epic Games, Inc. All Rights Reserved.
+
+#pragma once
+
+#include "CoreMinimal.h"
+
+class UMeshComponent;
+class IMeshPaintComponentAdapterFactory;
+class FReferenceCollector;
+
+class MESHPAINTINGTOOLSET_API FMeshPaintComponentAdapterFactory
+{
+public:
+	static TArray<TSharedPtr<IMeshPaintComponentAdapterFactory>> FactoryList;
+
+public:
+	static TSharedPtr<class IMeshPaintComponentAdapter> CreateAdapterForMesh(UMeshComponent* InComponent, int32 InPaintingMeshLODIndex);
+	static void InitializeAdapterGlobals();
+	static void AddReferencedObjectsGlobals(FReferenceCollector& Collector);
+	static void CleanupGlobals();
+};
\ No newline at end of file
diff --git a/Engine/Plugins/MeshPainting/Source/MeshPaintingToolset/Public/MeshPaintHelpers.h b/Engine/Plugins/MeshPainting/Source/MeshPaintingToolset/Public/MeshPaintHelpers.h
new file mode 100644
--- /dev/null
+++ b/Engine/Plugins/MeshPainting/Source/MeshPaintingToolset/Public/MeshPaintHelpers.h
@@ -0,0 +1,307 @@
+// Copyright 1998-2019 Epic Games, Inc. All Rights Reserved.
+
+#pragma once
+
+#include "CoreMinimal.h"
+#include "Engine/EngineTypes.h"
+#include "MeshPaintingToolsetTypes.h"
+#include "Kismet/BlueprintFunctionLibrary.h"
+#include "IMeshPaintComponentAdapter.h"
+#include "InteractiveToolManager.h"
+#include "InteractiveToolsContext.h"
+#include "Engine/StaticMesh.h"
+#include "MeshPaintHelpers.generated.h"
+
+class FMeshPaintParameters;
+class UImportVertexColorOptions;
+class UTexture2D;
+class UStaticMeshComponent;
+class USkeletalMesh;
+class IMeshPaintComponentAdapter;
+class UPaintBrushSettings;
+class FEditorViewportClient;
+class UMeshComponent;
+class USkeletalMeshComponent;
+class UViewportInteractor;
+class FViewport;
+class FPrimitiveDrawInterface;
+class FSceneView;
+struct FStaticMeshComponentLODInfo;
+class UMeshVertexPaintingToolProperties;
+
+enum class EMeshPaintDataColorViewMode : uint8;
+
+/** struct used to store the color data copied from mesh instance to mesh instance */
+struct FPerLODVertexColorData
+{
+	TArray< FColor > ColorsByIndex;
+	TMap<FVector, FColor> ColorsByPosition;
+};
+
+/** struct used to store the color data copied from mesh component to mesh component */
+struct FPerComponentVertexColorData
+{
+	FPerComponentVertexColorData(const UStaticMesh* InStaticMesh, int32 InComponentIndex)
+		: OriginalMesh(InStaticMesh)
+		, ComponentIndex(InComponentIndex)
+	{
+	}
+
+	/** We match up components by the mesh they use */
+	TWeakObjectPtr<const UStaticMesh> OriginalMesh;
+
+	/** We also match by component index */
+	int32 ComponentIndex;
+
+	/** Vertex colors by LOD */
+	TArray<FPerLODVertexColorData> PerLODVertexColorData;
+};
+
+/** Struct to hold MeshPaint settings on a per mesh basis */
+struct FInstanceTexturePaintSettings
+{
+	UTexture2D* SelectedTexture;
+	int32 SelectedUVChannel;
+
+	FInstanceTexturePaintSettings()
+		: SelectedTexture(nullptr)
+		, SelectedUVChannel(0)
+	{}
+	FInstanceTexturePaintSettings(UTexture2D* InSelectedTexture, int32 InSelectedUVSet)
+		: SelectedTexture(InSelectedTexture)
+		, SelectedUVChannel(InSelectedUVSet)
+	{}
+
+	void operator=(const FInstanceTexturePaintSettings& SrcSettings)
+	{
+		SelectedTexture = SrcSettings.SelectedTexture;
+		SelectedUVChannel = SrcSettings.SelectedUVChannel;
+	}
+};
+
+
+UENUM()
+enum class ETexturePaintWeightTypes : uint8
+{
+	/** Lerp Between Two Textures using Alpha Value */
+	AlphaLerp = 2 UMETA(DisplayName = "Alpha (Two Textures)"),
+
+	/** Weighting Three Textures according to Channels*/
+	RGB = 3 UMETA(DisplayName = "RGB (Three Textures)"),
+
+	/**  Weighting Four Textures according to Channels*/
+	ARGB = 4 UMETA(DisplayName = "ARGB (Four Textures)"),
+
+	/**  Weighting Five Textures according to Channels */
+	OneMinusARGB = 5 UMETA(DisplayName = "ARGB - 1 (Five Textures)")
+};
+
+UENUM()
+enum class ETexturePaintWeightIndex : uint8
+{
+	TextureOne = 0,
+	TextureTwo,
+	TextureThree,
+	TextureFour,
+	TextureFive
+};
+
+/** Parameters for paint actions, stored together for convenience */
+struct FPerVertexPaintActionArgs
+{
+	IMeshPaintComponentAdapter* Adapter;
+	UMeshVertexPaintingToolProperties* BrushProperties;
+	FVector CameraPosition;
+	FHitResult HitResult;
+	EMeshPaintModeAction Action;
+};
+
+/** Delegates used to call per-vertex/triangle actions */
+DECLARE_DELEGATE_TwoParams(FPerVertexPaintAction, FPerVertexPaintActionArgs& /*Args*/, int32 /*VertexIndex*/);
+DECLARE_DELEGATE_ThreeParams(FPerTrianglePaintAction, IMeshPaintComponentAdapter* /*Adapter*/, int32 /*TriangleIndex*/, const int32[3] /*Vertex Indices*/);
+
+class MESHPAINTINGTOOLSET_API UMeshPaintingToolset : public UBlueprintFunctionLibrary
+{
+public:
+	/** Removes vertex colors associated with the object */
+	static void RemoveInstanceVertexColors(UObject* Obj);
+
+	/** Removes vertex colors associated with the static mesh component */
+	static void RemoveComponentInstanceVertexColors(UStaticMeshComponent* StaticMeshComponent);
+	
+	/** Propagates per-instance vertex colors to the underlying Static Mesh for the given LOD Index */
+	static bool PropagateColorsToRawMesh(UStaticMesh* StaticMesh, int32 LODIndex, FStaticMeshComponentLODInfo& ComponentLODInfo);	
+
+	/** Retrieves the Vertex Color buffer size for the given LOD level in the Static Mesh */
+	static uint32 GetVertexColorBufferSize(UMeshComponent* MeshComponent, int32 LODIndex, bool bInstance);
+
+	/** Retrieves the vertex positions from the given LOD level in the Static Mesh */
+	static TArray<FVector> GetVerticesForLOD(const UStaticMesh* StaticMesh, int32 LODIndex);
+
+	/** Retrieves the vertex colors from the given LOD level in the Static Mesh */
+	static TArray<FColor> GetColorDataForLOD(const UStaticMesh* StaticMesh, int32 LODIndex);
+
+	/** Retrieves the per-instance vertex colors from the given LOD level in the StaticMeshComponent */
+	static TArray<FColor> GetInstanceColorDataForLOD(const UStaticMeshComponent* MeshComponent, int32 LODIndex);
+
+	/** Sets the specific (LOD Index) per-instance vertex colors for the given StaticMeshComponent to the supplied Color array */
+	static void SetInstanceColorDataForLOD(UStaticMeshComponent* MeshComponent, int32 LODIndex, const TArray<FColor>& Colors);	
+	
+	/** Sets the specific (LOD Index) per-instance vertex colors for the given StaticMeshComponent to a single Color value */
+	static void SetInstanceColorDataForLOD(UStaticMeshComponent* MeshComponent, int32 LODIndex, const FColor FillColor, const FColor MaskColor);
+	
+	/** Fills all vertex colors for all LODs found in the given mesh component with Fill Color */
+	static void FillStaticMeshVertexColors(UStaticMeshComponent* MeshComponent, int32 LODIndex, const FColor FillColor, const FColor MaskColor);
+	static void FillSkeletalMeshVertexColors(USkeletalMeshComponent* MeshComponent, int32 LODIndex, const FColor FillColor, const FColor MaskColor);
+	
+	/** Sets all vertex colors for a specific LOD level in the SkeletalMesh to FillColor */
+	static void SetColorDataForLOD(USkeletalMesh* SkeletalMesh, int32 LODIndex, const FColor FillColor, const FColor MaskColor);
+
+	static void ApplyFillWithMask(FColor& InOutColor, const FColor& MaskColor, const FColor& FillColor);
+
+		
+	/** Checks whether or not the given Viewport Client is a VR editor viewport client */
+	static bool IsInVRMode(const FEditorViewportClient* ViewportClient);
+
+	/** Forces the component to render LOD level at LODIndex instead of the view-based LOD level ( X = 0 means do not force the LOD, X > 0 means force the lod to X - 1 ) */
+	static void ForceRenderMeshLOD(UMeshComponent* Component, int32 LODIndex);
+
+	/** Clears all texture overrides for this component. */
+	static void ClearMeshTextureOverrides(const IMeshPaintComponentAdapter& GeometryInfo, UMeshComponent* InMeshComponent);
+
+	/** Applies vertex color painting found on LOD 0 to all lower LODs. */
+	static void ApplyVertexColorsToAllLODs(IMeshPaintComponentAdapter& GeometryInfo, UMeshComponent* InMeshComponent);
+
+	/** Applies the vertex colors found in LOD level 0 to all contained LOD levels in the StaticMeshComponent */
+	static void ApplyVertexColorsToAllLODs(IMeshPaintComponentAdapter& GeometryInfo, UStaticMeshComponent* StaticMeshComponent);
+
+	/** Applies the vertex colors found in LOD level 0 to all contained LOD levels in the SkeletalMeshComponent */
+	static void ApplyVertexColorsToAllLODs(IMeshPaintComponentAdapter& GeometryInfo, USkeletalMeshComponent* SkeletalMeshComponent);
+
+	/** Returns the number of Mesh LODs for the given MeshComponent */
+	static int32 GetNumberOfLODs(const UMeshComponent* MeshComponent);
+
+	/** OutNumLODs is set to number of Mesh LODs for the given MeshComponent and returns true, or returns false of given mesh component has no valid LODs */
+	static bool TryGetNumberOfLODs(const UMeshComponent* MeshComponent, int32& OutNumLODs);
+	
+	/** Returns the number of Texture Coordinates for the given MeshComponent */
+	static int32 GetNumberOfUVs(const UMeshComponent* MeshComponent, int32 LODIndex);
+
+	/** Checks whether or not the mesh components contains per lod colors (for all LODs)*/
+	static bool DoesMeshComponentContainPerLODColors(const UMeshComponent* MeshComponent);
+
+	/** Retrieves the number of bytes used to store the per-instance LOD vertex color data from the static mesh component */
+	static void GetInstanceColorDataInfo(const UStaticMeshComponent* StaticMeshComponent, int32 LODIndex, int32& OutTotalInstanceVertexColorBytes);
+
+	/** Given arguments for an action, and an action - retrieves influences vertices and applies Action to them */
+	static bool ApplyPerVertexPaintAction(FPerVertexPaintActionArgs& InArgs, FPerVertexPaintAction Action);
+
+	static bool GetPerVertexPaintInfluencedVertices(FPerVertexPaintActionArgs& InArgs, TSet<int32>& InfluencedVertices);
+
+	/** Given the adapter, settings and view-information retrieves influences triangles and applies Action to them */
+	static bool ApplyPerTrianglePaintAction(IMeshPaintComponentAdapter* Adapter, const FVector& CameraPosition, const FVector& HitPosition, const UMeshVertexPaintingToolProperties* Settings, FPerTrianglePaintAction Action);
+
+	/** Applies vertex painting to InOutvertexColor according to the given parameters  */
+	static bool PaintVertex(const FVector& InVertexPosition, const FMeshPaintParameters& InParams, FColor& InOutVertexColor);
+
+	/** Applies Vertex Color Painting according to the given parameters */
+	static void ApplyVertexColorPaint(const FMeshPaintParameters &InParams, const FLinearColor &OldColor, FLinearColor &NewColor, const float PaintAmount);
+
+	/** Applies Vertex Blend Weight Painting according to the given parameters */
+	static void ApplyVertexWeightPaint(const FMeshPaintParameters &InParams, const FLinearColor &OldColor, FLinearColor &NewColor, const float PaintAmount);
+
+	/** Generate texture weight color for given number of weights and the to-paint index */
+	static FLinearColor GenerateColorForTextureWeight(const int32 NumWeights, const int32 WeightIndex);
+
+	/** Computes the Paint power multiplier value */
+	static float ComputePaintMultiplier(float SquaredDistanceToVertex2D, float BrushStrength, float BrushInnerRadius, float BrushRadialFalloff, float BrushInnerDepth, float BrushDepthFallof, float VertexDepthToBrush);
+
+	/** Checks whether or not a point is influenced by the painting brush according to the given parameters*/
+	static bool IsPointInfluencedByBrush(const FVector& InPosition, const FMeshPaintParameters& InParams, float& OutSquaredDistanceToVertex2D, float& OutVertexDepthToBrush);
+
+	static bool IsPointInfluencedByBrush(const FVector2D& BrushSpacePosition, const float BrushRadiusSquared, float& OutInRangeValue);
+
+	template<typename T>
+	static void ApplyBrushToVertex(const FVector& VertexPosition, const FMatrix& InverseBrushMatrix, const float BrushRadius, const float BrushFalloffAmount, const float BrushStrength, const T& PaintValue, T& InOutValue);
+
+public:
+
+
+
+	/** Helper function to retrieve vertex color from a UTexture given a UVCoordinate */
+	static FColor PickVertexColorFromTextureData(const uint8* MipData, const FVector2D& UVCoordinate, const UTexture2D* Texture, const FColor ColorMask);	
+};
+
+template<typename T>
+void UMeshPaintingToolset::ApplyBrushToVertex(const FVector& VertexPosition, const FMatrix& InverseBrushMatrix, const float BrushRadius, const float BrushFalloffAmount, const float BrushStrength, const T& PaintValue, T& InOutValue)
+{
+	const FVector BrushSpacePosition = InverseBrushMatrix.TransformPosition(VertexPosition);
+	const FVector2D BrushSpacePosition2D(BrushSpacePosition.X, BrushSpacePosition.Y);
+		
+	float InfluencedValue = 0.0f;
+	if (IsPointInfluencedByBrush(BrushSpacePosition2D, BrushRadius * BrushRadius, InfluencedValue))
+	{
+		float InnerBrushRadius = BrushFalloffAmount * BrushRadius;
+		float PaintStrength = UMeshPaintingToolset::ComputePaintMultiplier(BrushSpacePosition2D.SizeSquared(), BrushStrength, InnerBrushRadius, BrushRadius - InnerBrushRadius, 1.0f, 1.0f, 1.0f);
+
+		const T OldValue = InOutValue;
+		InOutValue = FMath::LerpStable(OldValue, PaintValue, PaintStrength);
+	}	
+}
+
+UCLASS()
+class MESHPAINTINGTOOLSET_API UMeshToolsContext : public UInteractiveToolsContext
+{
+	GENERATED_BODY()
+
+public:
+	UMeshToolsContext();
+};
+
+UCLASS()
+class MESHPAINTINGTOOLSET_API UMeshToolManager : public UInteractiveToolManager
+{
+	GENERATED_BODY()
+
+public:
+	UMeshToolManager();
+	virtual void Shutdown() override;
+
+	/** Map of geometry adapters for each selected mesh component */
+	TMap<UMeshComponent*, TSharedPtr<IMeshPaintComponentAdapter>> GetComponentToAdapterMap() const;
+	TSharedPtr<IMeshPaintComponentAdapter> GetAdapterForComponent(UMeshComponent* InComponent);
+	void AddToComponentToAdapterMap(UMeshComponent* InComponent, TSharedPtr<IMeshPaintComponentAdapter> InAdapter);
+
+	TArray<UMeshComponent*> GetSelectedMeshComponents() const;
+	void AddSelectedMeshComponents(const TArray<UMeshComponent*>& InComponents);
+	bool FindHitResult(const FRay Ray, FHitResult& BestTraceResult);
+	void ClearSelectedMeshComponents();
+	TArray<UMeshComponent*> GetPaintableMeshComponents() const;
+	void AddPaintableMeshComponent(UMeshComponent* InComponent);
+	void ClearPaintableMeshComponents();
+	bool SelectionContainsValidAdapters() const;
+	TArray<FPerComponentVertexColorData> GetCopiedColorsByComponent() const;
+	void SetCopiedColorsByComponent(TArray<FPerComponentVertexColorData>& InCopiedColors);
+	void CacheSelectionData(const int32 PaintLODIndex, const int32 UVChannel);
+	void ResetState();
+	void Refresh();
+	bool SelectionContainsPerLODColors() const { return bSelectionContainsPerLODColors; }
+	void ClearSelectionLODColors() { bSelectionContainsPerLODColors = false; }
+
+public:
+	bool bNeedsRecache;
+
+
+private:
+	void CleanUp();
+
+private:
+	/** Map of geometry adapters for each selected mesh component */
+	TMap<UMeshComponent*, TSharedPtr<IMeshPaintComponentAdapter>> ComponentToAdapterMap;
+	TArray<UMeshComponent*> SelectedMeshComponents;
+	/** Mesh components within the current selection which are eligible for painting */
+	TArray<UMeshComponent*> PaintableComponents;
+	/** Contains copied vertex color data */
+	TArray<FPerComponentVertexColorData> CopiedColorsByComponent;
+	bool bSelectionContainsPerLODColors;
+};
\ No newline at end of file
diff --git a/Engine/Plugins/MeshPainting/Source/MeshPaintingToolset/Public/MeshPaintSplineMeshAdapter.h b/Engine/Plugins/MeshPainting/Source/MeshPaintingToolset/Public/MeshPaintSplineMeshAdapter.h
new file mode 100644
--- /dev/null
+++ b/Engine/Plugins/MeshPainting/Source/MeshPaintingToolset/Public/MeshPaintSplineMeshAdapter.h
@@ -0,0 +1,24 @@
+// Copyright 1998-2019 Epic Games, Inc. All Rights Reserved.
+
+#pragma once
+
+#include "CoreMinimal.h"
+#include "MeshPaintStaticMeshAdapter.h"
+
+//////////////////////////////////////////////////////////////////////////
+// FMeshPaintGeometryAdapterForSplineMeshes
+
+class MESHPAINTINGTOOLSET_API FMeshPaintSplineMeshComponentAdapter : public FMeshPaintStaticMeshComponentAdapter
+{
+public:
+	virtual bool InitializeVertexData() override;
+};
+
+//////////////////////////////////////////////////////////////////////////
+// FMeshPaintGeometryAdapterForSplineMeshesFactory
+
+class MESHPAINTINGTOOLSET_API FMeshPaintSplineMeshComponentAdapterFactory : public FMeshPaintStaticMeshComponentAdapterFactory
+{
+public:
+	virtual TSharedPtr<IMeshPaintComponentAdapter> Construct(class UMeshComponent* InComponent, int32 MeshLODIndex) const override;
+};
diff --git a/Engine/Plugins/MeshPainting/Source/MeshPaintingToolset/Public/MeshPaintingToolset.h b/Engine/Plugins/MeshPainting/Source/MeshPaintingToolset/Public/MeshPaintingToolset.h
new file mode 100644
--- /dev/null
+++ b/Engine/Plugins/MeshPainting/Source/MeshPaintingToolset/Public/MeshPaintingToolset.h
@@ -0,0 +1,5 @@
+// Copyright 1998-2019 Epic Games, Inc. All Rights Reserved.
+
+#pragma once
+
+#include "CoreMinimal.h"
diff --git a/Engine/Plugins/MeshPainting/Source/MeshPaintingToolset/Public/MeshPaintingToolsetTypes.h b/Engine/Plugins/MeshPainting/Source/MeshPaintingToolset/Public/MeshPaintingToolsetTypes.h
new file mode 100644
--- /dev/null
+++ b/Engine/Plugins/MeshPainting/Source/MeshPaintingToolset/Public/MeshPaintingToolsetTypes.h
@@ -0,0 +1,198 @@
+// Copyright 1998-2019 Epic Games, Inc. All Rights Reserved.
+
+#pragma once
+
+#include "Engine/Texture.h"
+#include "Engine/TextureRenderTarget2D.h"
+#include "Materials/MaterialInterface.h"
+
+#include "MeshPaintingToolsetTypes.generated.h"
+
+class FMeshPaintParameters;
+DECLARE_MULTICAST_DELEGATE_FourParams(FApplyVertexPaintData, const FMeshPaintParameters& /* &InParams*/, const FLinearColor& /*OldColor*/, FLinearColor& /*NewColor*/, const float /*PaintAmount*/);
+
+/** Mesh paint color view modes (somewhat maps to EVertexColorViewMode engine enum.) */
+UENUM()
+enum class EMeshPaintDataColorViewMode : uint8
+{
+	/** Normal view mode (vertex color visualization off) */
+	Normal UMETA(DisplayName = "Off"),
+
+	/** RGB only */
+	RGB UMETA(DisplayName = "RGB Channels"),
+
+	/** Alpha only */
+	Alpha UMETA(DisplayName = "Alpha Channel"),
+
+	/** Red only */
+	Red UMETA(DisplayName = "Red Channel"),
+
+	/** Green only */
+	Green UMETA(DisplayName = "Green Channel"),
+
+	/** Blue only */
+	Blue UMETA(DisplayName = "Blue Channel"),
+};
+
+/** Mesh painting action (paint, erase) */
+enum class EMeshPaintModeAction : uint8
+{
+	/** Paint (add color or increase blending weight) */
+	Paint,
+
+	/** Erase (remove color or decrease blending weight) */
+	Erase
+};
+
+/** Vertex paint target */
+UENUM()
+enum class EMeshVertexPaintModeTarget : uint8
+{
+	/** Paint the static mesh component instance in the level */
+	ComponentInstance,
+
+	/** Paint the actual static mesh asset */
+	Mesh
+};
+
+/** Mesh paint parameters */
+class FMeshPaintParameters
+{
+public:
+	EMeshPaintModeAction PaintAction;
+	FVector BrushPosition;
+	FVector BrushNormal;
+	FLinearColor BrushColor;
+	float SquaredBrushRadius;
+	float BrushRadialFalloffRange;
+	float InnerBrushRadius;
+	float BrushDepth;
+	float BrushDepthFalloffRange;
+	float InnerBrushDepth;
+	float BrushStrength;
+	FMatrix BrushToWorldMatrix;
+	FMatrix InverseBrushToWorldMatrix;
+	bool bWriteRed;
+	bool bWriteGreen;
+	bool bWriteBlue;
+	bool bWriteAlpha;
+	int32 TotalWeightCount;
+	int32 PaintWeightIndex;
+	int32 UVChannel;
+	FApplyVertexPaintData ApplyVertexDataDelegate;
+};
+
+/** Structure used to hold per-triangle data for texture painting */
+struct FTexturePaintTriangleInfo
+{
+	FVector TriVertices[3];
+	FVector2D TrianglePoints[3];
+	FVector2D TriUVs[3];
+};
+
+/** Structure used to house and compare Texture and UV channel pairs */
+struct FPaintableTexture
+{
+	UTexture*	Texture;
+	int32		UVChannelIndex;
+
+	FPaintableTexture(UTexture* InTexture = nullptr, uint32 InUVChannelIndex = 0)
+		: Texture(InTexture)
+		, UVChannelIndex(InUVChannelIndex)
+	{}
+
+	/** Overloaded equality operator for use with TArrays Contains method. */
+	bool operator==(const FPaintableTexture& rhs) const
+	{
+		return (Texture == rhs.Texture);
+		/* && (UVChannelIndex == rhs.UVChannelIndex);*/// if we compare UVChannel we would have to duplicate the texture
+	}
+};
+
+struct FPaintTexture2DData
+{
+	/** The original texture that we're painting */
+	UTexture2D* PaintingTexture2D;
+	bool bIsPaintingTexture2DModified;
+
+	/** A copy of the original texture we're painting, used for restoration. */
+	UTexture2D* PaintingTexture2DDuplicate;
+
+	/** Render target texture for painting */
+	UTextureRenderTarget2D* PaintRenderTargetTexture;
+
+	/** Render target texture used as an input while painting that contains a clone of the original image */
+	UTextureRenderTarget2D* CloneRenderTargetTexture;
+
+	/** List of materials we are painting on */
+	TArray< UMaterialInterface* > PaintingMaterials;
+
+	/** Default ctor */
+	FPaintTexture2DData() :
+		PaintingTexture2D(NULL),
+		bIsPaintingTexture2DModified(false),
+		PaintingTexture2DDuplicate(nullptr),
+		PaintRenderTargetTexture(nullptr),
+		CloneRenderTargetTexture(nullptr)
+	{}
+
+	FPaintTexture2DData(UTexture2D* InPaintingTexture2D, bool InbIsPaintingTexture2DModified = false) :
+		PaintingTexture2D(InPaintingTexture2D),
+		bIsPaintingTexture2DModified(InbIsPaintingTexture2DModified),
+		PaintRenderTargetTexture(nullptr),
+		CloneRenderTargetTexture(nullptr)
+	{}
+
+	/** Serializer */
+	void AddReferencedObjects(FReferenceCollector& Collector)
+	{
+		// @todo MeshPaint: We're relying on GC to clean up render targets, can we free up remote memory more quickly?
+		Collector.AddReferencedObject(PaintingTexture2D);
+		Collector.AddReferencedObject(PaintRenderTargetTexture);
+		Collector.AddReferencedObject(CloneRenderTargetTexture);
+		for (int32 Index = 0; Index < PaintingMaterials.Num(); Index++)
+		{
+			Collector.AddReferencedObject(PaintingMaterials[Index]);
+		}
+	}
+};
+
+namespace MeshPaintDefs
+{
+	// Design constraints
+
+	// Currently we never support more than five channels (R, G, B, A, OneMinusTotal)
+	static const int32 MaxSupportedPhysicalWeights = 4;
+	static const int32 MaxSupportedWeights = MaxSupportedPhysicalWeights + 1;
+}
+
+/**
+*  Wrapper to expose texture targets to WPF code.
+*/
+struct FTextureTargetListInfo
+{
+	UTexture2D* TextureData;
+	bool bIsSelected;
+	uint32 UndoCount;
+	uint32 UVChannelIndex;
+	FTextureTargetListInfo(UTexture2D* InTextureData, int32 InUVChannelIndex, bool InbIsSelected = false)
+		: TextureData(InTextureData)
+		, bIsSelected(InbIsSelected)
+		, UndoCount(0)
+		, UVChannelIndex(InUVChannelIndex)
+	{}
+};
+
+/**
+*  Wrapper to store which of a meshes materials is selected as well as the total number of materials.
+*/
+struct FMeshSelectedMaterialInfo
+{
+	int32 NumMaterials;
+	int32 SelectedMaterialIndex;
+
+	FMeshSelectedMaterialInfo(int32 InNumMaterials)
+		: NumMaterials(InNumMaterials)
+		, SelectedMaterialIndex(0)
+	{}
+};
\ No newline at end of file
diff --git a/Engine/Plugins/MeshPainting/Source/MeshPaintingToolset/Public/MeshSelect.h b/Engine/Plugins/MeshPainting/Source/MeshPaintingToolset/Public/MeshSelect.h
new file mode 100644
--- /dev/null
+++ b/Engine/Plugins/MeshPainting/Source/MeshPaintingToolset/Public/MeshSelect.h
@@ -0,0 +1,57 @@
+// Copyright 1998-2019 Epic Games, Inc. All Rights Reserved.
+
+#pragma once
+
+#include "CoreMinimal.h"
+#include "BaseTools/SingleClickTool.h"
+#include "MeshSelect.generated.h"
+
+
+
+/**
+ * Builder for UMeshClickToolBuilder
+ */
+UCLASS()
+class MESHPAINTINGTOOLSET_API UMeshClickToolBuilder : public USingleClickToolBuilder
+{
+	GENERATED_BODY()
+
+public:
+	virtual bool CanBuildTool(const FToolBuilderState& SceneState) const override;
+	virtual UInteractiveTool* BuildTool(const FToolBuilderState& SceneState) const override;
+};
+
+
+/**
+ * USingleClickTool is perhaps the simplest possible interactive tool. It simply
+ * reacts to default primary button clicks for the active device (eg left-mouse clicks).
+ *
+ * The function ::IsHitByClick() determines what is clickable by this Tool. The default is
+ * to return true, which means the click will activate anywhere (the Tool itself has no
+ * notion of Actors, Components, etc). You can override this function to, for example,
+ * filter out clicks that don't hit a target object, etc.
+ *
+ * The function ::OnClicked() implements the action that will occur when a click happens.
+ * You must override this to implement any kind of useful behavior.
+ */
+UCLASS()
+class MESHPAINTINGTOOLSET_API UMeshClickTool : public USingleClickTool
+{
+	GENERATED_BODY()
+
+public:
+	UMeshClickTool();
+
+	// USingleClickTool overrides
+	virtual	void Setup() override;
+	virtual void OnUpdateModifierState(int ModifierID, bool bIsOn) override;
+	virtual FInputRayHit IsHitByClick(const FInputDeviceRay& ClickPos) override;
+	virtual void OnClicked(const FInputDeviceRay& ClickPos) override;
+
+protected:
+	// flags used to identify modifier keys/buttons
+	static const int AdditiveSelectionModifier = 1;
+	bool bAddToSelectionSet;
+
+};
+
diff --git a/Engine/Plugins/MovieScene/MatineeToLevelSequence/MatineeToLevelSequence.uplugin b/Engine/Plugins/MovieScene/MatineeToLevelSequence/MatineeToLevelSequence.uplugin
--- a/Engine/Plugins/MovieScene/MatineeToLevelSequence/MatineeToLevelSequence.uplugin
+++ b/Engine/Plugins/MovieScene/MatineeToLevelSequence/MatineeToLevelSequence.uplugin
@@ -21,5 +21,11 @@
 			"Type" : "Editor",
 			"LoadingPhase" : "PostEngineInit"
 		}
-	]
-}
\ No newline at end of file
+	],
+    "Plugins": [
+        {
+            "Name" : "TemplateSequence",
+            "Enabled" : true
+        }
+    ]
+}
diff --git a/Engine/Plugins/MovieScene/MatineeToLevelSequence/Source/MatineeToLevelSequence/MatineeToLevelSequence.Build.cs b/Engine/Plugins/MovieScene/MatineeToLevelSequence/Source/MatineeToLevelSequence/MatineeToLevelSequence.Build.cs
--- a/Engine/Plugins/MovieScene/MatineeToLevelSequence/Source/MatineeToLevelSequence/MatineeToLevelSequence.Build.cs
+++ b/Engine/Plugins/MovieScene/MatineeToLevelSequence/Source/MatineeToLevelSequence/MatineeToLevelSequence.Build.cs
@@ -15,6 +15,7 @@ public class MatineeToLevelSequence : ModuleRules
         PrivateDependencyModuleNames.AddRange(
 			new string[] {
                 "LevelSequence",
+                "ContentBrowser",
 				"Core",
 				"CoreUObject",
                 "EditorStyle",
@@ -25,8 +26,10 @@ public class MatineeToLevelSequence : ModuleRules
 				"MovieSceneTracks",
                 "Slate",
                 "SlateCore",
+                "ToolMenus",
                 "UnrealEd",
 				"TimeManagement",
+                "TemplateSequence",
 				"Analytics"
             }
         );
