Fix For Mesh Painitng tools attempting to get the wedgemap from the wrong structure.

diff --git a/Engine/Plugins/MeshPainting/Source/MeshPaintingToolset/Private/MeshPaintHelpers.cpp b/Engine/Plugins/MeshPainting/Source/MeshPaintingToolset/Private/MeshPaintHelpers.cpp
--- a/Engine/Plugins/MeshPainting/Source/MeshPaintingToolset/Private/MeshPaintHelpers.cpp
+++ b/Engine/Plugins/MeshPainting/Source/MeshPaintingToolset/Private/MeshPaintHelpers.cpp
@@ -79,7 +79,7 @@ bool UMeshPaintingToolset::PropagateColorsToRawMesh(UStaticMesh* StaticMesh, int
 	FStaticMeshRenderData& RenderData = *StaticMesh->RenderData;
 	FStaticMeshLODResources& RenderModel = RenderData.LODResources[LODIndex];
 	FColorVertexBuffer& ColorVertexBuffer = *ComponentLODInfo.OverrideVertexColors;
-	if (RenderData.WedgeMap.Num() > 0 && ColorVertexBuffer.GetNumVertices() == RenderModel.GetNumVertices())
+	if (RenderModel.WedgeMap.Num() > 0 && ColorVertexBuffer.GetNumVertices() == RenderModel.GetNumVertices())
 	{
 		// Use the wedge map if it is available as it is lossless.
 		FMeshDescription* MeshDescription = StaticMesh->GetMeshDescription(LODIndex);
@@ -92,14 +92,14 @@ bool UMeshPaintingToolset::PropagateColorsToRawMesh(UStaticMesh* StaticMesh, int
 
 		FStaticMeshAttributes Attributes(*MeshDescription);
 		int32 NumWedges = MeshDescription->VertexInstances().Num();
-		if (RenderData.WedgeMap.Num() == NumWedges)
+		if (RenderModel.WedgeMap.Num() == NumWedges)
 		{
 			TVertexInstanceAttributesRef<FVector4> Colors = Attributes.GetVertexInstanceColors();
 			int32 VertexInstanceIndex = 0;
 			for (const FVertexInstanceID VertexInstanceID : MeshDescription->VertexInstances().GetElementIDs())
 			{
 				FLinearColor WedgeColor = FLinearColor::White;
-				int32 Index = RenderData.WedgeMap[VertexInstanceIndex];
+				int32 Index = RenderModel.WedgeMap[VertexInstanceIndex];
 				if (Index != INDEX_NONE)
 				{
 					WedgeColor = FLinearColor(ColorVertexBuffer.VertexColor(Index));
@@ -133,7 +133,7 @@ bool UMeshPaintingToolset::PropagateColorsToRawMesh(UStaticMesh* StaticMesh, int
 		int32 VertexIndex = 0;
 		for (const FVertexID VertexID : MeshDescription->Vertices().GetElementIDs())
 		{
-			VertexPositionsDup[VertexIndex] = VertexPositions[VertexID];
+			VertexPositionsDup[VertexIndex++] = VertexPositions[VertexID];
 		}
 		TempPositionVertexBuffer.Init(VertexPositionsDup);
 		RemapPaintedVertexColors(
