First pass on refactoring mesh paint to a plugin that uses the interactive tools framework

diff --git a/Engine/Plugins/Experimental/MeshModelingToolset/Source/MeshModelingTools/Private/DynamicMeshSculptTool.cpp b/Engine/Plugins/Experimental/MeshModelingToolset/Source/MeshModelingTools/Private/DynamicMeshSculptTool.cpp
--- a/Engine/Plugins/Experimental/MeshModelingToolset/Source/MeshModelingTools/Private/DynamicMeshSculptTool.cpp
+++ b/Engine/Plugins/Experimental/MeshModelingToolset/Source/MeshModelingTools/Private/DynamicMeshSculptTool.cpp
@@ -1216,7 +1216,7 @@ void UDynamicMeshSculptTool::Render(IToolsContextRenderAPI* RenderAPI)
 {
 	UMeshSurfacePointTool::Render(RenderAPI);
 
-	BrushIndicator->Update( (float)this->CurrentBrushRadius, this->LastBrushPosWorld, this->LastBrushPosNormalWorld );
+	BrushIndicator->Update( (float)this->CurrentBrushRadius, this->LastBrushPosWorld, this->LastBrushPosNormalWorld, 1.0f);
 
 	if (SculptProperties->PrimaryBrushType == EDynamicMeshSculptBrushType::FixedPlane)
 	{
diff --git a/Engine/Plugins/MeshPainting/Source/MeshPaintEditorMode/Private/BrushSettingsCustomization.cpp b/Engine/Plugins/MeshPainting/Source/MeshPaintEditorMode/Private/BrushSettingsCustomization.cpp
new file mode 100644
--- /dev/null
+++ b/Engine/Plugins/MeshPainting/Source/MeshPaintEditorMode/Private/BrushSettingsCustomization.cpp
@@ -0,0 +1,331 @@
+// Copyright 1998-2019 Epic Games, Inc. All Rights Reserved.
+
+#include "BrushSettingsCustomization.h"
+
+#include "DetailLayoutBuilder.h"
+#include "DetailCategoryBuilder.h"
+#include "IDetailPropertyRow.h"
+#include "IDetailChildrenBuilder.h"
+#include "DetailWidgetRow.h"
+#include "Widgets/Text/STextBlock.h"
+#include "PropertyRestriction.h"
+#include "Engine/Texture2D.h"
+#include "Widgets/SBoxPanel.h"
+#include "Widgets/Input/SCheckBox.h"
+#include "Widgets/Input/SNumericEntryBox.h"
+#include "PropertyCustomizationHelpers.h"
+#include "Widgets/Input/SButton.h"
+#include "Widgets/Layout/SBox.h"
+#include "ScopedTransaction.h"
+#include "MeshPaintHelpers.h"
+#include "MeshPaintColorBrush.h"
+#include "MeshPaintHelpers.h"
+#include "InteractiveToolManager.h"
+#include "Editor.h"
+#include "EditorModeManager.h"
+#include "MeshPaintMode.h"
+#include "Components/SkeletalMeshComponent.h"
+
+#define LOCTEXT_NAMESPACE "MeshPaintCustomization"
+
+TSharedRef<SHorizontalBox> CreateColorChannelWidget(TSharedRef<IPropertyHandle> ChannelProperty)
+{
+	return SNew(SHorizontalBox)
+		+ SHorizontalBox::Slot()
+		.AutoWidth()
+		[
+			ChannelProperty->CreatePropertyValueWidget()
+		]
+	+ SHorizontalBox::Slot()
+		.AutoWidth()
+		.Padding(4.0f, 0.0f, 0.0f, 0.0f)
+		[
+			ChannelProperty->CreatePropertyNameWidget()
+		];
+}
+
+
+TSharedRef<IDetailCustomization> FVertexPaintingSettingsCustomization::MakeInstance()
+{
+	return MakeShareable(new FVertexPaintingSettingsCustomization);
+}
+
+void FVertexPaintingSettingsCustomization::CustomizeDetails(IDetailLayoutBuilder& DetailLayout)
+{	
+	IDetailCategoryBuilder& VertexCategory = DetailLayout.EditCategory(TEXT("VertexPainting"));
+
+	VertexCategory.AddCustomRow(NSLOCTEXT("VertexPaintSettings", "InstanceColorSize", "Instance Color Size"))
+		.WholeRowContent()
+		[
+			SNew(STextBlock)
+			.Text_Lambda([]() -> FText { return FText::Format(FTextFormat::FromString(TEXT("Instance Color Size: {0} KB")), UMeshPaintMode::GetMeshPaintMode()->GetCachedVertexDataSize() / 1024.f); })
+		];
+
+	/** Creates a custom widget row containing all color channel flags */	
+	TSharedRef<IPropertyHandle> PaintColor = DetailLayout.GetProperty("PaintColor", UMeshVertexPaintingToolProperties::StaticClass());
+	PaintColor->MarkHiddenByCustomization();
+	TSharedRef<IPropertyHandle> EraseColor = DetailLayout.GetProperty("EraseColor", UMeshVertexPaintingToolProperties::StaticClass());
+	EraseColor->MarkHiddenByCustomization();
+
+	// Customize paint color with a swap button
+	{
+		TSharedPtr<SWidget> NameWidget;
+		TSharedPtr<SWidget> ValueWidget;
+		
+		IDetailPropertyRow& PaintColorProp = VertexCategory.AddProperty(PaintColor);
+		PaintColorProp.GetDefaultWidgets(NameWidget, ValueWidget, false);
+		FDetailWidgetRow& Row = PaintColorProp.CustomWidget(true);
+		Row.NameContent()
+		[
+			NameWidget.ToSharedRef()
+		];
+
+		Row.ValueContent()
+		.MinDesiredWidth(250)
+		.MaxDesiredWidth(0)
+		[
+			SNew(SHorizontalBox)
+			+ SHorizontalBox::Slot()
+			.FillWidth(1.0)
+			.HAlign(HAlign_Left)
+			[
+				SNew(SBox)
+				.WidthOverride(250.f)
+				[
+					ValueWidget.ToSharedRef()
+				]
+			]
+			+ SHorizontalBox::Slot()
+			.HAlign(HAlign_Center)
+			.AutoWidth()
+			[
+				SNew(SButton)
+				.ButtonStyle(FEditorStyle::Get(), "HoverHintOnly")
+				.ToolTipText(NSLOCTEXT("VertexPaintSettings", "SwapColors", "Swap Paint and Erase Colors"))
+				.HAlign(HAlign_Center)
+				.VAlign(VAlign_Center)
+				.OnClicked(this, &FVertexPaintingSettingsCustomization::OnSwapColorsClicked, PaintColor, EraseColor)
+				.ContentPadding(0)
+				[
+					SNew(SImage).Image(FEditorStyle::GetBrush("MeshPaint.Swap"))
+				]
+			]
+		];
+	}
+
+	{
+		IDetailPropertyRow& EraseColorProp = VertexCategory.AddProperty(EraseColor);
+
+		TSharedPtr<SWidget> NameWidget;
+		TSharedPtr<SWidget> ValueWidget;
+		
+		FDetailWidgetRow& Row = EraseColorProp.CustomWidget(true);
+		Row.ValueContent().MinDesiredWidth(250 - 16.f);
+		EraseColorProp.GetDefaultWidgets(NameWidget, ValueWidget, Row, false);
+	}
+
+}
+
+void FVertexPaintingSettingsCustomization::OnTextureWeightTypeChanged(TSharedRef<IPropertyHandle> WeightTypeProperty, TSharedRef<IPropertyHandle> PaintWeightProperty, TSharedRef<IPropertyHandle> EraseWeightProperty)
+{
+	UEnum* ImportTypeEnum = StaticEnum<ETexturePaintWeightIndex>();
+	uint8 EnumValue = 0;
+	WeightTypeProperty->GetValue(EnumValue);	
+
+	BlendPaintEnumRestriction->RemoveAll();
+	for (uint8 EnumIndex = 0; EnumIndex < (ImportTypeEnum->GetMaxEnumValue() + 1); ++EnumIndex)
+	{
+		if ((EnumIndex + 1) > EnumValue)
+		{
+			FString EnumName = ImportTypeEnum->GetNameByValue(EnumIndex).ToString();
+			EnumName.RemoveFromStart("ETexturePaintIndex::");
+			BlendPaintEnumRestriction->AddDisabledValue(EnumName);
+		}
+	}
+
+	uint8 Value = 0;
+	PaintWeightProperty->GetValue(Value);
+	Value = FMath::Clamp<uint8>(Value, 0, EnumValue - 1);
+	PaintWeightProperty->SetValue(Value);
+
+	EraseWeightProperty->GetValue(Value);
+	Value = FMath::Clamp<uint8>(Value, 0, EnumValue - 1);
+	EraseWeightProperty->SetValue(Value);
+}
+
+FReply FVertexPaintingSettingsCustomization::OnSwapColorsClicked(TSharedRef<IPropertyHandle> PaintColor, TSharedRef<IPropertyHandle> EraseColor)
+{
+	FScopedTransaction Transaction(NSLOCTEXT("VertexPaintSettings", "SwapColorsTransation", "Swap paint and erase colors"));
+
+	UMeshVertexPaintingToolProperties* Settings = UMeshPaintMode::GetVertexToolProperties();
+	if (Settings)
+	{
+		Settings->Modify();
+		
+		FLinearColor TempPaintColor = Settings->PaintColor;
+		Settings->PaintColor = Settings->EraseColor;
+		Settings->EraseColor = TempPaintColor;
+		PaintColor->NotifyPostChange(EPropertyChangeType::ValueSet);
+		EraseColor->NotifyPostChange(EPropertyChangeType::ValueSet);
+	}
+	return FReply::Handled();
+}
+
+
+
+TSharedRef<IDetailCustomization> FColorPaintingSettingsCustomization::MakeInstance()
+{
+	return MakeShareable(new FColorPaintingSettingsCustomization);
+}
+
+void FColorPaintingSettingsCustomization::CustomizeDetails(IDetailLayoutBuilder& DetailLayout)
+{
+	FVertexPaintingSettingsCustomization::CustomizeDetails(DetailLayout);
+
+	IDetailCategoryBuilder& ColorCategory = DetailLayout.EditCategory(TEXT("ColorPainting"));
+
+	/** Creates a custom widget row containing all color channel flags */
+	TSharedRef<IPropertyHandle> RedChannel = DetailLayout.GetProperty("bWriteRed", UMeshColorPaintingToolProperties::StaticClass());
+	RedChannel->MarkHiddenByCustomization();
+	TSharedRef<IPropertyHandle> GreenChannel = DetailLayout.GetProperty("bWriteGreen", UMeshColorPaintingToolProperties::StaticClass());
+	GreenChannel->MarkHiddenByCustomization();
+	TSharedRef<IPropertyHandle> BlueChannel = DetailLayout.GetProperty("bWriteBlue", UMeshColorPaintingToolProperties::StaticClass());
+	BlueChannel->MarkHiddenByCustomization();
+	TSharedRef<IPropertyHandle> AlphaChannel = DetailLayout.GetProperty("bWriteAlpha", UMeshColorPaintingToolProperties::StaticClass());
+	AlphaChannel->MarkHiddenByCustomization();
+	TArray<TSharedRef<IPropertyHandle>> Channels = { RedChannel, GreenChannel, BlueChannel, AlphaChannel };
+	TSharedPtr<SHorizontalBox> ChannelsWidget;
+
+	ColorCategory.AddCustomRow(NSLOCTEXT("VertexPaintSettings", "ChannelLabel", "Channels"))
+		.NameContent()
+		[
+			SNew(STextBlock)
+			.Text(NSLOCTEXT("VertexPaintSettings", "ChannelsLabel", "Channels"))
+			.ToolTipText(NSLOCTEXT("VertexPaintSettings", "ChannelsToolTip", "Colors Channels which should be influenced during Painting."))
+		]
+		.ValueContent()
+		.MaxDesiredWidth(250.0f)
+		[
+			SAssignNew(ChannelsWidget, SHorizontalBox)
+		];
+
+	for (TSharedRef<IPropertyHandle> Channel : Channels)
+	{
+		ChannelsWidget->AddSlot()
+			.AutoWidth()
+			.Padding(0.0f, 0.0f, 4.0f, 0.0f)
+			[
+				CreateColorChannelWidget(Channel)
+			];
+	}
+
+	static FText RestrictReason = NSLOCTEXT("VertexPaintSettings", "TextureIndexRestriction", "Unable to paint this Texture, change Texture Weight Type");
+	BlendPaintEnumRestriction = MakeShareable(new FPropertyRestriction(RestrictReason));
+
+	/** Add custom row for painting on specific LOD level with callbacks to the painter to update the data */
+	TSharedRef<IPropertyHandle> LODPaintingEnabled = DetailLayout.GetProperty("bPaintOnSpecificLOD", UMeshColorPaintingToolProperties::StaticClass());
+	LODPaintingEnabled->MarkHiddenByCustomization();
+	TSharedRef<IPropertyHandle> LODPaintingIndex = DetailLayout.GetProperty("LODIndex", UMeshColorPaintingToolProperties::StaticClass());
+	LODPaintingIndex->MarkHiddenByCustomization();
+	TSharedPtr<SWidget> LODIndexWidget = LODPaintingIndex->CreatePropertyValueWidget();
+
+	ColorCategory.AddCustomRow(NSLOCTEXT("LODPainting", "LODPaintingLabel", "LOD Model Painting"))
+		.NameContent()
+		[
+			SNew(STextBlock)
+			.Text(NSLOCTEXT("LODPainting", "LODPaintingSetupLabel", "LOD Model Painting"))
+			.ToolTipText(NSLOCTEXT("LODPainting", "LODPaintingSetupToolTip", "Allows for Painting Vertex Colors on Specific LOD Models."))
+		]
+	.ValueContent()
+		[
+			SNew(SHorizontalBox)
+			+ SHorizontalBox::Slot()
+		.Padding(0.0f, 0.0f, 4.0f, 0.0f)
+		.AutoWidth()
+		[
+			SNew(SCheckBox)
+				.IsChecked_Lambda([=]() -> ECheckBoxState { return (UMeshPaintMode::GetColorToolProperties()->bPaintOnSpecificLOD ? ECheckBoxState::Checked : ECheckBoxState::Unchecked); })
+				.OnCheckStateChanged(FOnCheckStateChanged::CreateLambda([=](ECheckBoxState State) { 
+					if (UMeshColorPaintingTool* ColorBrush = Cast<UMeshColorPaintingTool>(UMeshPaintMode::GetMeshPaintMode()->GetToolManager()->GetActiveTool(EToolSide::Left)))
+					{
+						ColorBrush->LODPaintStateChanged(State == ECheckBoxState::Checked);
+					}
+				}))
+		]
+	+ SHorizontalBox::Slot()
+		.Padding(0.0f, 0.0f, 4.0f, 0.0f)
+		[
+			SNew(SNumericEntryBox<int32>)
+			.IsEnabled_Lambda([=]() -> bool { return UMeshPaintMode::GetColorToolProperties()->bPaintOnSpecificLOD;  })
+			.AllowSpin(true)
+			.Value_Lambda([=]() -> int32 { return UMeshPaintMode::GetColorToolProperties()->LODIndex; })
+			.MinValue(0)
+			.MaxValue_Lambda([=]() -> int32 { 
+					if (UMeshColorPaintingTool* ColorBrush = Cast<UMeshColorPaintingTool>(UMeshPaintMode::GetMeshPaintMode()->GetToolManager()->GetActiveTool(EToolSide::Left)))
+					{
+						return ColorBrush->GetMaxLODIndexToPaint();
+					}
+					return INT_MAX;
+				})
+			.MaxSliderValue_Lambda([=]() -> int32 { 
+					if (UMeshColorPaintingTool* ColorBrush = Cast<UMeshColorPaintingTool>(UMeshPaintMode::GetMeshPaintMode()->GetToolManager()->GetActiveTool(EToolSide::Left)))
+					{
+						return ColorBrush->GetMaxLODIndexToPaint();
+					}
+					return INT_MAX;
+				})
+			.OnValueChanged(SNumericEntryBox<int32>::FOnValueChanged::CreateLambda([=](int32 Value) { UMeshPaintMode::GetColorToolProperties()->LODIndex = Value; }))
+			.OnValueCommitted(SNumericEntryBox<int32>::FOnValueCommitted::CreateLambda([=](int32 Value, ETextCommit::Type CommitType) { 
+					UMeshPaintMode::GetColorToolProperties()->LODIndex = Value; 
+					if (UMeshColorPaintingTool* ColorBrush = Cast<UMeshColorPaintingTool>(UMeshPaintMode::GetMeshPaintMode()->GetToolManager()->GetActiveTool(EToolSide::Left)))
+					{
+						ColorBrush->PaintLODChanged();
+					}
+				}))
+		]
+		];
+
+	ColorCategory.AddCustomRow(NSLOCTEXT("LODPainting", "LODPaintingLabel", "LOD Model Painting"))
+		.WholeRowContent()
+		[
+			SNew(SBorder)
+			.Visibility_Lambda([this]() -> EVisibility
+			{
+				if (UMeshColorPaintingToolProperties* ColorProperties = UMeshPaintMode::GetColorToolProperties())
+				{
+					return ColorProperties->bPaintOnSpecificLOD ? EVisibility::Collapsed : EVisibility::Visible;
+				}
+				return EVisibility::Collapsed;
+			})
+		.Padding(FMargin(4.0f))
+		.BorderImage(FEditorStyle::GetBrush("SettingsEditor.CheckoutWarningBorder"))
+		.BorderBackgroundColor(FColor(166, 137, 0))
+		[
+			SNew(STextBlock)
+			.AutoWrapText(true)
+			.Font(IDetailLayoutBuilder::GetDetailFont())
+
+		.Text_Lambda([this]() -> FText
+		{
+			static const FText SkelMeshNotificationText = LOCTEXT("SkelMeshAssetPaintInfo", "Paint is propagated to Skeletal Mesh Asset(s)");
+			static const FText StaticMeshNotificationText = LOCTEXT("StaticMeshAssetPaintInfo", "Paint is applied to all LODs");
+
+			const bool bSkelMeshText = UMeshPaintMode::GetMeshPaintMode()->GetSelectedComponents<USkeletalMeshComponent>().Num() > 0;
+			const bool bLODPaintText = UMeshPaintMode::GetColorToolProperties() ?  !UMeshPaintMode::GetColorToolProperties()->bPaintOnSpecificLOD : false;
+			return FText::Format(FTextFormat::FromString(TEXT("{0}{1}{2}")), bSkelMeshText ? SkelMeshNotificationText : FText::GetEmpty(), bSkelMeshText && bLODPaintText ? FText::FromString(TEXT("\n")) : FText::GetEmpty(), bLODPaintText ? StaticMeshNotificationText : FText::GetEmpty());
+		})
+		]
+		];
+}
+
+TSharedRef<IDetailCustomization> FWeightPaintingSettingsCustomization::MakeInstance()
+{
+	return MakeShareable(new FWeightPaintingSettingsCustomization);
+}
+
+void FWeightPaintingSettingsCustomization::CustomizeDetails(IDetailLayoutBuilder& DetailLayout)
+{
+	FVertexPaintingSettingsCustomization::CustomizeDetails(DetailLayout);
+}
+
+#undef LOCTEXT_NAMESPACE
\ No newline at end of file
diff --git a/Engine/Plugins/MeshPainting/Source/MeshPaintEditorMode/Private/BrushSettingsCustomization.h b/Engine/Plugins/MeshPainting/Source/MeshPaintEditorMode/Private/BrushSettingsCustomization.h
new file mode 100644
--- /dev/null
+++ b/Engine/Plugins/MeshPainting/Source/MeshPaintEditorMode/Private/BrushSettingsCustomization.h
@@ -0,0 +1,65 @@
+// Copyright 1998-2019 Epic Games, Inc. All Rights Reserved.
+
+#pragma once
+
+#include "IDetailCustomization.h"
+#include "IPropertyTypeCustomization.h"
+#include "IDetailRootObjectCustomization.h"
+#include "PropertyRestriction.h"
+#include "Widgets/SWidget.h"
+
+class SHorizontalBox;
+class UMeshVertexPaintingToolProperties;
+
+
+class FVertexPaintingSettingsCustomization : public IDetailCustomization
+{
+public:
+	static TSharedRef<IDetailCustomization> MakeInstance();
+
+	/** IPropertyTypeCustomization interface */
+	virtual void CustomizeDetails(IDetailLayoutBuilder& DetailLayout) override;
+protected:
+	/** Callback for when texture weight type changed so we can update restrictions */
+	void OnTextureWeightTypeChanged(TSharedRef<IPropertyHandle> WeightTypeProperty, TSharedRef<IPropertyHandle> PaintWeightProperty, TSharedRef<IPropertyHandle> EraseWeightProperty);
+	FReply OnSwapColorsClicked(TSharedRef<IPropertyHandle> PaintColor, TSharedRef<IPropertyHandle> EraseColor);
+
+	/** Property restriction applied to blend paint enum dropdown box */
+	TSharedPtr<FPropertyRestriction> BlendPaintEnumRestriction;
+};
+
+class FColorPaintingSettingsCustomization : public FVertexPaintingSettingsCustomization
+{
+public:
+	static TSharedRef<IDetailCustomization> MakeInstance();
+
+	/** IPropertyTypeCustomization interface */
+	virtual void CustomizeDetails(IDetailLayoutBuilder& DetailLayout) override;
+};
+
+class FWeightPaintingSettingsCustomization : public FVertexPaintingSettingsCustomization
+{
+public:
+	static TSharedRef<IDetailCustomization> MakeInstance();
+
+	/** IPropertyTypeCustomization interface */
+	virtual void CustomizeDetails(IDetailLayoutBuilder& DetailLayout) override;
+};
+
+// class FTexturePaintSettingsCustomization : public IPropertyTypeCustomization
+// {
+// public:
+// 	static TSharedRef<IPropertyTypeCustomization> MakeInstance();
+// 
+// 	/** IPropertyTypeCustomization interface */
+// 	virtual void CustomizeHeader(TSharedRef<IPropertyHandle> PropertyHandle, FDetailWidgetRow& HeaderRow, IPropertyTypeCustomizationUtils& CustomizationUtils) override {};
+// 	virtual void CustomizeChildren(TSharedRef<IPropertyHandle> PropertyHandle, IDetailChildrenBuilder& ChildBuilder, IPropertyTypeCustomizationUtils& CustomizationUtils) override;
+// protected:
+// 	/** Cached instance of the painter */
+// 	FPaintModePainter* MeshPainter;
+// 	/** Cached instance of the texture painting settings */
+// 	FTexturePaintSettings* PaintSettings;
+// };
+
+
+	
\ No newline at end of file
diff --git a/Engine/Plugins/MeshPainting/Source/MeshPaintEditorMode/Private/MeshPaintEditorMode.cpp b/Engine/Plugins/MeshPainting/Source/MeshPaintEditorMode/Private/MeshPaintEditorMode.cpp
new file mode 100644
--- /dev/null
+++ b/Engine/Plugins/MeshPainting/Source/MeshPaintEditorMode/Private/MeshPaintEditorMode.cpp
@@ -0,0 +1,147 @@
+// Copyright 1998-2019 Epic Games, Inc. All Rights Reserved.
+
+#include "MeshPaintEditorMode.h"
+#include "EditorModeRegistry.h"
+#include "Modules/ModuleManager.h"
+#include "MeshPaintMode.h"
+#include "EditorModeManager.h"
+#include "MeshPaintModeSettings.h"
+#include "MeshPaintModeCommands.h"
+#include "BrushSettingsCustomization.h"
+#include "PropertyEditorModule.h"
+#include "IMeshPaintGeometryAdapterFactory.h"
+#include "MeshPaintAdapterFactory.h"
+#include "MeshPaintSplineMeshAdapter.h"
+#include "MeshPaintStaticMeshAdapter.h"
+#include "MeshPaintSkeletalMeshAdapter.h"
+#include "Settings/LevelEditorMiscSettings.h"
+ 
+
+#define LOCTEXT_NAMESPACE "MeshPaintMode"
+
+class FMeshPaintModeModule : public IModuleInterface
+{
+public:
+	FMeshPaintModeModule() 
+	{
+	}
+
+	// FModuleInterface overrides
+	virtual void StartupModule() override;
+	virtual void ShutdownModule() override;
+	virtual bool SupportsDynamicReloading() override
+	{
+		return true;
+	}
+
+protected:
+
+	/** Stores the FEditorMode ID of the associated editor mode */
+	static inline FEditorModeID GetEditorModeID()
+	{
+		static FEditorModeID MeshPaintModeFeatureName = FName( TEXT( "MeshPaintMode" ) );
+		return MeshPaintModeFeatureName;
+	}
+
+	void Register();
+	void OnPostEngineInit();
+	void Unregister();
+
+
+	void OnMeshPaintModeButtonClicked();
+	bool IsMeshPaintModeButtonEnabled();
+	void RegisterGeometryAdapterFactory(TSharedRef<IMeshPaintGeometryAdapterFactory> Factory);
+	void UnregisterGeometryAdapterFactory(TSharedRef<IMeshPaintGeometryAdapterFactory> Factory);
+};
+
+void FMeshPaintModeModule::Register()
+{
+	RegisterGeometryAdapterFactory(MakeShareable(new FMeshPaintGeometryAdapterForSplineMeshesFactory));
+	RegisterGeometryAdapterFactory(MakeShareable(new FMeshPaintGeometryAdapterForStaticMeshesFactory));
+	RegisterGeometryAdapterFactory(MakeShareable(new FMeshPaintGeometryAdapterForSkeletalMeshesFactory));
+
+
+	FEditorModeRegistry::Get().RegisterScriptableMode<UMeshPaintMode>(
+		GetEditorModeID(),
+		LOCTEXT("ModeName", "Mesh Paint"),
+		FSlateIcon(FEditorStyle::GetStyleSetName(), "LevelEditor.MeshPaintMode", "LevelEditor.MeshPaintMode.Small"),
+		true,
+		600
+		);
+
+	FPropertyEditorModule& PropertyModule = FModuleManager::LoadModuleChecked<FPropertyEditorModule>("PropertyEditor");
+	PropertyModule.RegisterCustomClassLayout("MeshColorPaintingToolProperties", FOnGetDetailCustomizationInstance::CreateStatic(&FColorPaintingSettingsCustomization::MakeInstance));
+	PropertyModule.RegisterCustomClassLayout("MeshWeightPaintingToolProperties", FOnGetDetailCustomizationInstance::CreateStatic(&FWeightPaintingSettingsCustomization::MakeInstance));
+}
+
+
+void FMeshPaintModeModule::OnPostEngineInit()
+{
+	if (!GetDefault<ULevelEditorMiscSettings>()->bEnableLegacyMeshPaintMode)
+	{
+		Register();
+		FMeshPaintingToolActionCommands::RegisterAllToolActions();
+		FMeshPaintEditorModeCommands::Register();
+	}
+}
+
+void FMeshPaintModeModule::StartupModule()
+{
+	FCoreDelegates::OnPostEngineInit.AddRaw(this, &FMeshPaintModeModule::OnPostEngineInit);
+}
+
+
+void FMeshPaintModeModule::ShutdownModule()
+{
+	
+	FCoreDelegates::OnPostEngineInit.RemoveAll(this);
+
+	FMeshPaintingToolActionCommands::UnregisterAllToolActions();
+	FMeshPaintEditorModeCommands::Unregister();
+
+	// This function may be called during shutdown to clean up your module.  For modules that support dynamic reloading,
+	// we call this function before unloading the module.
+	Unregister();
+	
+}
+
+
+
+
+void FMeshPaintModeModule::Unregister()
+{
+	FEditorModeRegistry::Get().UnregisterMode(GetEditorModeID());
+}
+
+
+void FMeshPaintModeModule::OnMeshPaintModeButtonClicked()
+{
+	// *Important* - activate the mode first since FEditorModeTools::DeactivateMode will
+	// activate the default mode when the stack becomes empty, resulting in multiple active visible modes.
+	GLevelEditorModeTools().ActivateMode(GetEditorModeID());
+
+	// Find and disable any other 'visible' modes since we only ever allow one of those active at a time.
+	GLevelEditorModeTools().DeactivateOtherVisibleModes(GetEditorModeID());
+
+}
+
+
+
+bool FMeshPaintModeModule::IsMeshPaintModeButtonEnabled( )
+{
+	return true;
+}
+
+void  FMeshPaintModeModule::RegisterGeometryAdapterFactory(TSharedRef<IMeshPaintGeometryAdapterFactory> Factory)
+{
+	FMeshPaintAdapterFactory::FactoryList.Add(Factory);
+}
+
+void  FMeshPaintModeModule::UnregisterGeometryAdapterFactory(TSharedRef<IMeshPaintGeometryAdapterFactory> Factory)
+{
+	FMeshPaintAdapterFactory::FactoryList.Remove(Factory);
+}
+
+IMPLEMENT_MODULE(FMeshPaintModeModule, MeshPaintEditorMode)
+
+#undef LOCTEXT_NAMESPACE
diff --git a/Engine/Plugins/MeshPainting/Source/MeshPaintEditorMode/Private/MeshPaintMode.cpp b/Engine/Plugins/MeshPainting/Source/MeshPaintEditorMode/Private/MeshPaintMode.cpp
new file mode 100644
--- /dev/null
+++ b/Engine/Plugins/MeshPainting/Source/MeshPaintEditorMode/Private/MeshPaintMode.cpp
@@ -0,0 +1,686 @@
+// Copyright 1998-2019 Epic Games, Inc. All Rights Reserved.
+
+#include "MeshPaintMode.h"
+#include "MeshPaintModeCommands.h"
+#include "MeshPaintModeToolkit.h"
+#include "Toolkits/ToolkitManager.h"
+#include "Framework/Commands/UICommandList.h"
+#include "MeshPaintModeSettings.h"
+#include "Framework/Commands/Commands.h"
+#include "MultiSelectionTool.h"
+#include "InteractiveToolManager.h"
+#include "InteractiveToolsContext.h"
+#include "Components/MeshComponent.h"
+#include "Dialogs/Dialogs.h"
+#include "Components/StaticMeshComponent.h"
+#include "MeshPaintHelpers.h"
+#include "Components/SkeletalMeshComponent.h"
+#include "ScopedTransaction.h"
+#include "PackageTools.h"
+#include "MeshPaintHelpers.h"
+#include "Engine/Selection.h"
+#include "MeshPaintHelpers.h"
+#include "MeshPaintAdapterFactory.h"
+#include "MeshPaintModeHelpers.h"
+#include "MeshSelect.h"
+
+
+#define LOCTEXT_NAMESPACE "MeshPaintMode"
+
+FName UMeshPaintMode::MeshPaintMode_Color = FName(TEXT("Color"));
+FName UMeshPaintMode::MeshPaintMode_Weights = FName(TEXT("Weights"));
+FName UMeshPaintMode::MeshPaintMode_Texture = FName(TEXT("Texture"));
+
+UMeshVertexPaintingToolProperties* UMeshPaintMode::GetVertexToolProperties()
+{
+	if (GLevelEditorModeTools().GetActiveScriptableMode("MeshPaintMode")
+		&& GLevelEditorModeTools().GetActiveScriptableMode("MeshPaintMode")->GetToolManager()->GetActiveTool(EToolSide::Mouse))
+	{
+		TArray<UObject*> PropertyArray = GLevelEditorModeTools().GetActiveScriptableMode("MeshPaintMode")->GetToolManager()->GetActiveTool(EToolSide::Mouse)->GetToolProperties();
+		for(UObject* Property : PropertyArray)
+		{
+			if (UMeshVertexPaintingToolProperties* VertexProperties = Cast<UMeshVertexPaintingToolProperties>(Property))
+			{
+				return VertexProperties;
+			}
+		}
+	}
+	return nullptr;
+}
+
+UMeshColorPaintingToolProperties* UMeshPaintMode::GetColorToolProperties()
+{
+	if (GLevelEditorModeTools().GetActiveScriptableMode("MeshPaintMode")
+		&& GLevelEditorModeTools().GetActiveScriptableMode("MeshPaintMode")->GetToolManager()->GetActiveTool(EToolSide::Mouse))
+	{
+		TArray<UObject*> PropertyArray = GLevelEditorModeTools().GetActiveScriptableMode("MeshPaintMode")->GetToolManager()->GetActiveTool(EToolSide::Mouse)->GetToolProperties();
+		for (UObject* Property : PropertyArray)
+		{
+			if (UMeshColorPaintingToolProperties* ColorProperties = Cast<UMeshColorPaintingToolProperties>(Property))
+			{
+				return ColorProperties;
+			}
+		}
+	}
+	return nullptr;
+}
+
+UMeshWeightPaintingToolProperties* UMeshPaintMode::GetWeightToolProperties()
+{
+	if (GLevelEditorModeTools().GetActiveScriptableMode("MeshPaintMode")
+		&& GLevelEditorModeTools().GetActiveScriptableMode("MeshPaintMode")->GetToolManager()->GetActiveTool(EToolSide::Mouse))
+	{
+		TArray<UObject*> PropertyArray = GLevelEditorModeTools().GetActiveScriptableMode("MeshPaintMode")->GetToolManager()->GetActiveTool(EToolSide::Mouse)->GetToolProperties();
+		for (UObject* Property : PropertyArray)
+		{
+			if (UMeshWeightPaintingToolProperties* WeightProperties = Cast<UMeshWeightPaintingToolProperties>(Property))
+			{
+				return WeightProperties;
+			}
+		}
+	}
+	return nullptr;
+}
+
+UMeshPaintMode* UMeshPaintMode::GetMeshPaintMode()
+{
+	return Cast<UMeshPaintMode>(GLevelEditorModeTools().GetActiveScriptableMode("MeshPaintMode"));
+}
+
+UMeshPaintMode::UMeshPaintMode()
+	:Super()
+{
+	SettingsClass = UMeshPaintModeSettings::StaticClass();
+	ToolsContextClass = UMeshToolsContext::StaticClass();
+	CurrentPaletteName = MeshPaintMode_Color;
+}
+
+void UMeshPaintMode::Enter()
+{
+	Super::Enter();
+	ModeSettings = Cast<UMeshPaintModeSettings>(SettingsObject);
+	FMeshPaintEditorModeCommands ToolManagerCommands = FMeshPaintEditorModeCommands::Get();
+	RegisterTool(ToolManagerCommands.Select, TEXT("MeshClickTool"), NewObject<UMeshClickToolBuilder>());
+	RegisterTool(ToolManagerCommands.ColorPaint, TEXT("ColorBrushTool"), NewObject<UMeshColorPaintingToolBuilder>());
+	RegisterTool(ToolManagerCommands.WeightPaint, TEXT("WeightBrushTool"), NewObject<UMeshWeightPaintingToolBuilder>());
+
+	UpdateSelectedMeshes();
+
+	ActivateDefaultTool();
+}
+
+void UMeshPaintMode::Exit()
+{
+	const FMeshPaintEditorModeCommands& Commands = FMeshPaintEditorModeCommands::Get();
+	const TSharedRef<FUICommandList>& CommandList = Toolkit->GetToolkitCommands();
+	for (const TSharedPtr<const FUICommandInfo> Action : Commands.Commands[UMeshPaintMode::MeshPaintMode_Color])
+	{
+		CommandList->UnmapAction(Action);
+	}
+	for (const TSharedPtr<const FUICommandInfo> Action : Commands.Commands[UMeshPaintMode::MeshPaintMode_Weights])
+	{
+		CommandList->UnmapAction(Action);
+	}
+	for (const TSharedPtr<const FUICommandInfo> Action : Commands.Commands[UMeshPaintMode::MeshPaintMode_Weights])
+	{
+		CommandList->UnmapAction(Action);
+	}
+	Super::Exit();
+}
+
+void UMeshPaintMode::CreateToolkit()
+{
+	if (!Toolkit.IsValid())
+	{
+		Toolkit = MakeShareable(new FMeshPaintModeToolkit);
+		Toolkit->Init(Owner->GetToolkitHost());
+	}
+
+	// Register UI commands
+	BindCommands();
+
+	Super::CreateToolkit();
+}
+
+void UMeshPaintMode::Tick(FEditorViewportClient* ViewportClient, float DeltaTime)
+{
+	UEdMode::Tick(ViewportClient, DeltaTime);
+
+	if (ViewportClient->IsPerspective())
+	{
+		// Make sure perspective viewports are still set to real-time
+		UMeshPaintModeHelpers::SetRealtimeViewport(true);
+
+		// Set viewport show flags		
+		UMeshPaintModeHelpers::SetViewportColorMode(ModeSettings->ColorViewMode, ViewportClient);
+	}
+
+
+	if (bRecacheVertexDataSize)
+	{
+		UpdateCachedVertexDataSize();
+	}
+}
+
+TMap<FName, TArray<TSharedPtr<FUICommandInfo>>> UMeshPaintMode::GetModeCommands() const
+{
+	return FMeshPaintEditorModeCommands::GetCommands();
+}
+
+void UMeshPaintMode::BindCommands()
+{
+	const FMeshPaintEditorModeCommands& Commands = FMeshPaintEditorModeCommands::Get();
+	const TSharedRef<FUICommandList>& CommandList = Toolkit->GetToolkitCommands();
+
+	auto IsInValidPaintMode = [this]() -> bool { return (GetVertexToolProperties() != nullptr) && Cast<UMeshToolManager>(GetToolManager())->SelectionContainsValidAdapters(); };
+	CommandList->MapAction(Commands.Fill, 
+		FUIAction(FExecuteAction::CreateUObject(this, &UMeshPaintMode::FillWithVertexColor),
+		FCanExecuteAction::CreateLambda(IsInValidPaintMode)));
+
+	CommandList->MapAction(Commands.Propagate,
+		FUIAction(
+			FExecuteAction::CreateUObject(this, &UMeshPaintMode::PropagateVertexColorsToAsset),
+			FCanExecuteAction::CreateUObject(this, &UMeshPaintMode::CanPropagateVertexColors)
+		));
+
+	auto IsAValidMeshComponentSelected = [this]() -> bool { return (GetSelectedComponents<UMeshComponent>().Num() == 1) && Cast<UMeshToolManager>(GetToolManager())->SelectionContainsValidAdapters(); };
+	CommandList->MapAction(Commands.Import,
+		FUIAction(
+			FExecuteAction::CreateUObject(this, &UMeshPaintMode::ImportVertexColors),
+			FCanExecuteAction::CreateLambda(IsAValidMeshComponentSelected)
+		));
+
+	CommandList->MapAction(Commands.Save,
+		FUIAction(
+			FExecuteAction::CreateUObject(this, &UMeshPaintMode::SavePaintedAssets),
+			FCanExecuteAction::CreateUObject(this, &UMeshPaintMode::CanSaveMeshPackages)
+		));
+
+	CommandList->MapAction(Commands.Copy,
+		FUIAction(
+			FExecuteAction::CreateUObject(this, &UMeshPaintMode::CopyVertexColors),
+			FCanExecuteAction::CreateUObject(this, &UMeshPaintMode::CanCopyInstanceVertexColors)
+		));
+
+	CommandList->MapAction(Commands.Paste,
+		FUIAction(
+			FExecuteAction::CreateUObject(this, &UMeshPaintMode::PasteVertexColors),
+			FCanExecuteAction::CreateUObject(this, &UMeshPaintMode::CanPasteInstanceVertexColors)
+		));
+
+	CommandList->MapAction(Commands.Remove,
+		FUIAction(
+			FExecuteAction::CreateUObject(this, &UMeshPaintMode::RemoveVertexColors),
+			FCanExecuteAction::CreateUObject(this, &UMeshPaintMode::CanRemoveInstanceColors)
+		));
+
+	CommandList->MapAction(Commands.Fix,
+		FUIAction(
+			FExecuteAction::CreateUObject(this, &UMeshPaintMode::FixVertexColors),
+			FCanExecuteAction::CreateUObject(this, &UMeshPaintMode::DoesRequireVertexColorsFixup)
+		));
+
+	CommandList->MapAction(Commands.PropagateVertexColorsToLODs,
+		FUIAction(
+			FExecuteAction::CreateUObject(this, &UMeshPaintMode::PropagateVertexColorsToLODs),
+			FCanExecuteAction::CreateUObject(this, &UMeshPaintMode::CanPropagateVertexColorsToLODs)
+		));
+
+}
+
+
+
+
+void UMeshPaintMode::OnToolStarted(UInteractiveToolManager* Manager, UInteractiveTool* Tool)
+{
+	FMeshPaintingToolActionCommands::UpdateToolCommandBinding(Tool, ToolCommandList, false);
+
+	if (UMeshVertexPaintingTool* VertexPaintingTool = Cast<UMeshVertexPaintingTool>(GetToolManager()->GetActiveTool(EToolSide::Left)))
+	{
+		VertexPaintingTool->OnPaintingFinished().BindUObject(this, &UMeshPaintMode::UpdateCachedVertexDataSize);
+	}
+}
+
+void UMeshPaintMode::OnToolEnded(UInteractiveToolManager* Manager, UInteractiveTool* Tool)
+{
+	FMeshPaintingToolActionCommands::UpdateToolCommandBinding(Tool, ToolCommandList, true);
+	// First update your bindings, then call the base behavior
+	Super::OnToolEnded(Manager, Tool);
+}
+
+bool UMeshPaintMode::InputAxis(FEditorViewportClient* InViewportClient, FViewport* Viewport, int32 ControllerId, FKey Key, float Delta, float DeltaTime)
+{
+	bool bSelectingOnly = (GetToolManager()->GetActiveTool(EToolSide::Mouse)) ? GetToolManager()->GetActiveTool(EToolSide::Mouse)->IsA<USingleClickTool>() : false;
+
+	// If we are painting (not selecting), don't allow panning with the mouse.
+	return !bSelectingOnly && !InViewportClient->IsAltPressed();
+}
+
+void UMeshPaintMode::ActorSelectionChangeNotify()
+{
+	UpdateSelectedMeshes();
+}
+
+void UMeshPaintMode::UpdateSelectedMeshes()
+{
+	if (UMeshToolManager* MeshToolManager = Cast<UMeshToolManager>(GetToolManager()))
+	{
+		MeshToolManager->ResetState();
+		const TArray<UMeshComponent*> CurrentMeshComponents = GetSelectedComponents<UMeshComponent>();
+		MeshToolManager->AddSelectedMeshComponents(CurrentMeshComponents);
+		MeshToolManager->bNeedsRecache = true;
+	}
+	bRecacheVertexDataSize = true;
+}
+
+void UMeshPaintMode::FillWithVertexColor()
+{
+	FScopedTransaction Transaction(LOCTEXT("LevelMeshPainter_TransactionFillInstColors", "Filling Per-Instance Vertex Colors"));
+	const TArray<UMeshComponent*> MeshComponents = GetSelectedComponents<UMeshComponent>();
+
+	static const bool bConvertSRGB = false;
+	FColor FillColor = FColor::White;
+	FColor MaskColor = FColor::White;
+
+	if (GetToolManager()->GetActiveTool(EToolSide::Mouse)->IsA<UMeshWeightPaintingTool>())
+	{
+		FillColor = UMeshPaintingToolset::GenerateColorForTextureWeight((int32)GetWeightToolProperties()->TextureWeightType, (int32)GetWeightToolProperties()->PaintTextureWeightIndex).ToFColor(bConvertSRGB);
+	}
+	else if (UMeshColorPaintingToolProperties* ColorProperties = GetColorToolProperties())
+	{
+		FillColor = ColorProperties->PaintColor.ToFColor(bConvertSRGB);
+		MaskColor.R = ColorProperties->bWriteRed ? 255 : 0;
+		MaskColor.G = ColorProperties->bWriteGreen ? 255 : 0;
+		MaskColor.B = ColorProperties->bWriteBlue ? 255 : 0;
+		MaskColor.A = ColorProperties->bWriteAlpha ? 255 : 0;
+	}
+
+	TUniquePtr< FComponentReregisterContext > ComponentReregisterContext;
+	/** Fill each mesh component with the given vertex color */
+	for (UMeshComponent* Component : MeshComponents)
+	{
+		checkf(Component != nullptr, TEXT("Invalid Mesh Component"));
+		Component->Modify();
+		ComponentReregisterContext = MakeUnique<FComponentReregisterContext>(Component);
+
+		TSharedPtr<IMeshPaintGeometryAdapter> MeshAdapter = Cast<UMeshToolManager>(GetToolManager())->GetAdapterForComponent(Component);
+		if (MeshAdapter)
+		{
+			MeshAdapter->PreEdit();
+		}
+
+		if (UMeshColorPaintingToolProperties* ColorProperties = UMeshPaintMode::GetColorToolProperties())
+		{
+			const bool bPaintOnSpecificLOD = ColorProperties ? ColorProperties->bPaintOnSpecificLOD : false;
+
+			if (Component->IsA<UStaticMeshComponent>())
+			{
+				UMeshPaintingToolset::FillStaticMeshVertexColors(Cast<UStaticMeshComponent>(Component), bPaintOnSpecificLOD ? ColorProperties->LODIndex : -1, FillColor, MaskColor);
+			}
+			else if (Component->IsA<USkeletalMeshComponent>())
+			{
+				UMeshPaintingToolset::FillSkeletalMeshVertexColors(Cast<USkeletalMeshComponent>(Component), bPaintOnSpecificLOD ? ColorProperties->LODIndex : -1, FillColor, MaskColor);
+			}
+		}
+
+		if (MeshAdapter)
+		{
+			MeshAdapter->PostEdit();
+		}
+	}
+}
+
+
+void UMeshPaintMode::PropagateVertexColorsToAsset()
+{
+	const TArray<UStaticMeshComponent*> StaticMeshComponents = GetSelectedComponents<UStaticMeshComponent>();
+	FSuppressableWarningDialog::FSetupInfo SetupInfo(LOCTEXT("PushInstanceVertexColorsPrompt_Message", "Copying the instance vertex colors to the source mesh will replace any of the source mesh's pre-existing vertex colors and affect every instance of the source mesh."),
+		LOCTEXT("PushInstanceVertexColorsPrompt_Title", "Warning: Copying vertex data overwrites all instances"), "Warning_PushInstanceVertexColorsPrompt");
+
+	SetupInfo.ConfirmText = LOCTEXT("PushInstanceVertexColorsPrompt_ConfirmText", "Continue");
+	SetupInfo.CancelText = LOCTEXT("PushInstanceVertexColorsPrompt_CancelText", "Abort");
+	SetupInfo.CheckBoxText = LOCTEXT("PushInstanceVertexColorsPrompt_CheckBoxText", "Always copy vertex colors without prompting");
+
+	FSuppressableWarningDialog VertexColorCopyWarning(SetupInfo);
+
+	// Prompt the user to see if they really want to push the vert colors to the source mesh and to explain
+	// the ramifications of doing so. This uses a suppressible dialog so that the user has the choice to always ignore the warning.
+	if (VertexColorCopyWarning.ShowModal() != FSuppressableWarningDialog::Cancel)
+	{
+		FScopedTransaction Transaction(LOCTEXT("LevelMeshPainter_TransactionPropogateColors", "Propagating Vertex Colors To Source Meshes"));
+		UMeshPaintModeHelpers::PropagateVertexColors(StaticMeshComponents);
+	}
+}
+
+bool UMeshPaintMode::CanPropagateVertexColors() const
+{
+	// Check whether or not our selected Static Mesh Components contain instance based vertex colors (only these can be propagated to the base mesh)
+	int32 NumInstanceVertexColorBytes = 0;
+
+	TArray<UStaticMesh*> StaticMeshes;
+	TArray<UStaticMeshComponent*> StaticMeshComponents = GetSelectedComponents<UStaticMeshComponent>();
+	return UMeshPaintModeHelpers::CanPropagateVertexColors(StaticMeshComponents, StaticMeshes, NumInstanceVertexColorBytes);
+
+}
+
+void UMeshPaintMode::ImportVertexColors()
+{
+	const TArray<UMeshComponent*> MeshComponents = GetSelectedComponents<UMeshComponent>();
+	if (MeshComponents.Num() == 1)
+	{
+		/** Import vertex color to single selected mesh component */
+		FScopedTransaction Transaction(LOCTEXT("LevelMeshPainter_TransactionImportColors", "Importing Vertex Colors From Texture"));
+		UMeshPaintModeHelpers::ImportVertexColorsFromTexture(MeshComponents[0]);
+	}
+}
+
+void UMeshPaintMode::SavePaintedAssets()
+{
+	const TArray<UStaticMeshComponent*> StaticMeshComponents = GetSelectedComponents<UStaticMeshComponent>();
+	const TArray<USkeletalMeshComponent*> SkeletalMeshComponents = GetSelectedComponents<USkeletalMeshComponent>();
+
+	/** Try and save outstanding dirty packages for currently selected mesh components */
+	TArray<UObject*> ObjectsToSave;
+	for (UStaticMeshComponent* StaticMeshComponent : StaticMeshComponents)
+	{
+		if (StaticMeshComponent && StaticMeshComponent->GetStaticMesh())
+		{
+			ObjectsToSave.Add(StaticMeshComponent->GetStaticMesh());
+		}
+	}
+
+	for (USkeletalMeshComponent* SkeletalMeshComponent : SkeletalMeshComponents)
+	{
+		if (SkeletalMeshComponent && SkeletalMeshComponent->SkeletalMesh)
+		{
+			ObjectsToSave.Add(SkeletalMeshComponent->SkeletalMesh);
+		}
+	}
+
+	if (ObjectsToSave.Num() > 0)
+	{
+		UPackageTools::SavePackagesForObjects(ObjectsToSave);
+	}
+}
+
+bool UMeshPaintMode::CanSaveMeshPackages() const
+{
+	// Check whether or not any of our selected mesh components contain mesh objects which require saving
+	TArray<UMeshComponent*> Components = GetSelectedComponents<UMeshComponent>();
+
+	bool bValid = false;
+
+	for (UMeshComponent* Component : Components)
+	{
+		UObject* Object = nullptr;
+		if (UStaticMeshComponent* StaticMeshComponent = Cast<UStaticMeshComponent>(Component))
+		{
+			Object = StaticMeshComponent->GetStaticMesh();
+		}
+		else if (USkeletalMeshComponent* SkeletalMeshComponent = Cast<USkeletalMeshComponent>(Component))
+		{
+			Object = SkeletalMeshComponent->SkeletalMesh;
+		}
+
+		if (Object != nullptr && Object->GetOutermost()->IsDirty())
+		{
+			bValid = true;
+			break;
+		}
+	}
+
+	return bValid;
+}
+
+bool UMeshPaintMode::CanRemoveInstanceColors() const
+{
+	const TArray<UStaticMeshComponent*> StaticMeshComponents = GetSelectedComponents<UStaticMeshComponent>();
+	int32 PaintingMeshLODIndex = 0;
+	if (UMeshColorPaintingToolProperties* ColorProperties = UMeshPaintMode::GetColorToolProperties())
+	{
+		PaintingMeshLODIndex = ColorProperties->bPaintOnSpecificLOD ? ColorProperties->LODIndex : 0;
+	}
+	int32 NumValidMeshes = 0;
+	// Retrieve per instance vertex color information (only valid if the component contains actual instance vertex colors)
+	for (UStaticMeshComponent* Component : StaticMeshComponents)
+	{
+		if (Component != nullptr && Component->GetStaticMesh() != nullptr && Component->GetStaticMesh()->GetNumLODs() > (int32)PaintingMeshLODIndex)
+		{
+			uint32 BufferSize = UMeshPaintingToolset::GetVertexColorBufferSize(Component, PaintingMeshLODIndex, true);
+
+			if (BufferSize > 0)
+			{
+				++NumValidMeshes;
+			}
+		}
+	}
+
+	return (NumValidMeshes != 0);
+}
+
+bool UMeshPaintMode::CanPasteInstanceVertexColors() const
+{
+	const TArray<UStaticMeshComponent*> StaticMeshComponents = GetSelectedComponents<UStaticMeshComponent>();
+	const TArray<FPerComponentVertexColorData> CopiedColorsByComponent = Cast<UMeshToolManager>(GetToolManager())->GetCopiedColorsByComponent();
+	return UMeshPaintModeHelpers::CanPasteInstanceVertexColors(StaticMeshComponents, CopiedColorsByComponent);
+}
+
+bool UMeshPaintMode::CanCopyInstanceVertexColors() const
+{
+	const TArray<UStaticMeshComponent*> StaticMeshComponents = GetSelectedComponents<UStaticMeshComponent>();
+	int32 PaintingMeshLODIndex = 0;
+	if (UMeshColorPaintingToolProperties* ColorProperties = UMeshPaintMode::GetColorToolProperties())
+	{
+		PaintingMeshLODIndex = ColorProperties->bPaintOnSpecificLOD ? ColorProperties->LODIndex : 0;
+	}
+
+	return UMeshPaintModeHelpers::CanCopyInstanceVertexColors(StaticMeshComponents, PaintingMeshLODIndex);
+
+}
+
+bool UMeshPaintMode::CanPropagateVertexColorsToLODs() const
+{
+	bool bPaintOnSpecificLOD = false;
+	if (UMeshColorPaintingToolProperties* ColorProperties = UMeshPaintMode::GetColorToolProperties())
+	{
+		bPaintOnSpecificLOD = ColorProperties ? ColorProperties->bPaintOnSpecificLOD : false;
+	}
+	// Can propagate when the mesh contains per-lod vertex colors or when we are not painting to a specific lod
+	const bool bSelectionContainsPerLODColors = Cast<UMeshToolManager>(GetToolManager())->SelectionContainsPerLODColors();
+	return bSelectionContainsPerLODColors || bPaintOnSpecificLOD;
+}
+
+void UMeshPaintMode::CopyVertexColors()
+{
+	const TArray<UStaticMeshComponent*> StaticMeshComponents = GetSelectedComponents<UStaticMeshComponent>();
+	TArray<FPerComponentVertexColorData> CopiedColorsByComponent;
+	UMeshPaintModeHelpers::CopyVertexColors(StaticMeshComponents, CopiedColorsByComponent);
+	Cast<UMeshToolManager>(GetToolManager())->SetCopiedColorsByComponent(CopiedColorsByComponent);
+
+}
+
+void UMeshPaintMode::PasteVertexColors()
+{
+	FScopedTransaction Transaction(LOCTEXT("LevelMeshPainter_TransactionPasteInstColors", "Pasting Per-Instance Vertex Colors"));
+	const TArray<UStaticMeshComponent*> StaticMeshComponents = GetSelectedComponents<UStaticMeshComponent>();
+	TArray<FPerComponentVertexColorData> CopiedColorsByComponent = Cast<UMeshToolManager>(GetToolManager())->GetCopiedColorsByComponent();
+	UMeshPaintModeHelpers::PasteVertexColors(StaticMeshComponents, CopiedColorsByComponent);
+	UpdateCachedVertexDataSize();
+}
+
+void UMeshPaintMode::FixVertexColors()
+{
+	FScopedTransaction Transaction(LOCTEXT("LevelMeshPainter_TransactionFixInstColors", "Fixing Per-Instance Vertex Colors"));
+	const TArray<UStaticMeshComponent*> StaticMeshComponents = GetSelectedComponents<UStaticMeshComponent>();
+	for (UStaticMeshComponent* Component : StaticMeshComponents)
+	{
+		Component->FixupOverrideColorsIfNecessary();
+	}
+}
+
+bool UMeshPaintMode::DoesRequireVertexColorsFixup() const
+{
+	const TArray<UStaticMeshComponent*> StaticMeshComponents = GetSelectedComponents<UStaticMeshComponent>();
+	bool bAnyMeshNeedsFixing = false;
+	/** Check if there are any static mesh components which require fixing */
+	for (UStaticMeshComponent* Component : StaticMeshComponents)
+	{
+		bAnyMeshNeedsFixing |= Component->RequiresOverrideVertexColorsFixup();
+	}
+
+	return bAnyMeshNeedsFixing;
+}
+
+void UMeshPaintMode::RemoveVertexColors()
+{
+	FScopedTransaction Transaction(LOCTEXT("LevelMeshPainter_TransactionRemoveInstColors", "Removing Per-Instance Vertex Colors"));
+	const TArray<UStaticMeshComponent*> StaticMeshComponents = GetSelectedComponents<UStaticMeshComponent>();
+	for (UStaticMeshComponent* Component : StaticMeshComponents)
+	{
+		UMeshPaintingToolset::RemoveComponentInstanceVertexColors(Component);
+	}
+
+	UpdateCachedVertexDataSize();
+}
+
+
+void UMeshPaintMode::PropagateVertexColorsToLODs()
+{
+	//Only show the lost data warning if there is actually some data to lose
+	bool bAbortChange = false;
+	UMeshToolManager* MeshToolManager = Cast<UMeshToolManager>(GetToolManager());
+	TArray<UMeshComponent*> PaintableComponents = MeshToolManager->GetPaintableMeshComponents();
+	const bool bSelectionContainsPerLODColors = MeshToolManager->SelectionContainsPerLODColors();
+	if (bSelectionContainsPerLODColors)
+	{
+		//Warn the user they will lose custom painting data
+		FSuppressableWarningDialog::FSetupInfo SetupInfo(LOCTEXT("LooseLowersLODsVertexColorsPrompt_Message", "Propagating Vertex Colors from LOD0 to all lower LODs. This mean all lower LODs custom vertex painting will be lost."),
+			LOCTEXT("LooseLowersLODsVertexColorsPrompt_Title", "Warning: Lowers LODs custom vertex painting will be lost!"), "Warning_LooseLowersLODsVertexColorsPrompt");
+
+		SetupInfo.ConfirmText = LOCTEXT("LooseLowersLODsVertexColorsPrompt_ConfirmText", "Continue");
+		SetupInfo.CancelText = LOCTEXT("LooseLowersLODsVertexColorsPrompt_CancelText", "Abort");
+		SetupInfo.CheckBoxText = LOCTEXT("LooseLowersLODsVertexColorsPrompt_CheckBoxText", "Always copy vertex colors without prompting");
+
+		FSuppressableWarningDialog LooseLowersLODsVertexColorsWarning(SetupInfo);
+	
+		// Prompt the user to see if they really want to propagate the base lod vert colors to the lowers LODs.
+		if (LooseLowersLODsVertexColorsWarning.ShowModal() == FSuppressableWarningDialog::Cancel)
+		{
+			bAbortChange = true;
+		}
+		else
+		{
+			// Reset the state flag as we'll be removing all per-lod colors 
+			MeshToolManager->ClearSelectionLODColors();
+			UMeshPaintModeHelpers::RemovePerLODColors(PaintableComponents);
+		}
+	}
+
+	//The user cancel the change, avoid changing the value
+	if (bAbortChange)
+	{
+		return;
+	}
+
+	for (UMeshComponent* SelectedComponent : PaintableComponents)
+	{
+		if (SelectedComponent)
+		{
+			TSharedPtr<IMeshPaintGeometryAdapter> MeshAdapter = Cast<UMeshToolManager>(GetToolManager())->GetAdapterForComponent(SelectedComponent);
+			UMeshPaintingToolset::ApplyVertexColorsToAllLODs(*MeshAdapter, SelectedComponent);
+			FComponentReregisterContext ReregisterContext(SelectedComponent);
+		}
+	}
+
+	MeshToolManager->Refresh();
+}
+
+template<typename ComponentClass>
+TArray<ComponentClass*> UMeshPaintMode::GetSelectedComponents() const
+{
+	TArray<ComponentClass*> Components;
+	FToolBuilderState SelectionState;
+	GetToolManager()->GetContextQueriesAPI()->GetCurrentSelectionState(SelectionState);
+//	if (PaintSettings->PaintMode == EPaintMode::Textures)
+	{
+
+		for (int32 SelectionIndex = 0; SelectionIndex < SelectionState.SelectedComponents.Num(); ++SelectionIndex)
+		{
+			ComponentClass* SelectedComponent = Cast<ComponentClass>(SelectionState.SelectedComponents[SelectionIndex]);
+			if (SelectedComponent)
+			{
+				Components.AddUnique(SelectedComponent);
+			}
+		}
+	}
+
+	if (Components.Num() == 0)
+	{
+		for (int32 SelectionIndex = 0; SelectionIndex < SelectionState.SelectedActors.Num(); ++SelectionIndex)
+		{
+			AActor* SelectedActor = Cast<AActor>(SelectionState.SelectedActors[SelectionIndex]);
+			if (SelectedActor)
+			{
+				TInlineComponentArray<ComponentClass*> ActorComponents;
+				SelectedActor->GetComponents(ActorComponents);
+				for (ComponentClass* Component : ActorComponents)
+				{
+					Components.AddUnique(Component);
+				}
+			}
+		}
+	}
+
+	return Components;
+}
+
+template TArray<UStaticMeshComponent*> UMeshPaintMode::GetSelectedComponents<UStaticMeshComponent>() const;
+template TArray<USkeletalMeshComponent*> UMeshPaintMode::GetSelectedComponents<USkeletalMeshComponent>() const;
+template TArray<UMeshComponent*> UMeshPaintMode::GetSelectedComponents<UMeshComponent>() const;
+
+
+
+
+void UMeshPaintMode::UpdateCachedVertexDataSize()
+{
+	CachedVertexDataSize = 0;
+
+	const bool bInstance = true;
+	if (UMeshToolManager* MeshToolManager = Cast<UMeshToolManager>(GetToolManager()))
+	{
+		for (UMeshComponent* SelectedComponent : MeshToolManager->GetPaintableMeshComponents())
+		{
+			int32 NumLODs = UMeshPaintingToolset::GetNumberOfLODs(SelectedComponent);
+			for (int32 LODIndex = 0; LODIndex < NumLODs; ++LODIndex)
+			{
+				CachedVertexDataSize += UMeshPaintingToolset::GetVertexColorBufferSize(SelectedComponent, LODIndex, bInstance);
+			}
+		}
+	}
+}
+
+
+
+void UMeshPaintMode::ActivateDefaultTool()
+{
+	if (CurrentPaletteName == UMeshPaintMode::MeshPaintMode_Color)
+	{
+		ToolsContext->StartTool(EToolSide::Mouse, TEXT("MeshClickTool"));
+	}
+	if (CurrentPaletteName == UMeshPaintMode::MeshPaintMode_Weights)
+	{
+		ToolsContext->StartTool(EToolSide::Mouse, TEXT("MeshClickTool"));
+	}
+	if (CurrentPaletteName == UMeshPaintMode::MeshPaintMode_Weights)
+	{
+
+	}
+}
+
+void UMeshPaintMode::UpdateOnPaletteChange()
+{
+	UpdateSelectedMeshes();
+}
+
+#undef LOCTEXT_NAMESPACE
+
diff --git a/Engine/Plugins/MeshPainting/Source/MeshPaintEditorMode/Private/MeshPaintMode.h b/Engine/Plugins/MeshPainting/Source/MeshPaintEditorMode/Private/MeshPaintMode.h
new file mode 100644
--- /dev/null
+++ b/Engine/Plugins/MeshPainting/Source/MeshPaintEditorMode/Private/MeshPaintMode.h
@@ -0,0 +1,101 @@
+// Copyright 1998-2019 Epic Games, Inc. All Rights Reserved.
+
+#pragma once
+
+#include "Tools/UEdMode.h"
+#include "Framework/Commands/UIAction.h"
+#include "Framework/Commands/UICommandList.h"
+#include "Containers/ArrayView.h"
+#include "LevelEditorViewport.h"	// For FTrackingTransaction
+#include "UnrealEdMisc.h"
+#include "Framework/Commands/UIAction.h"
+#include "Framework/Commands/UICommandInfo.h"
+#include "Framework/Commands/UICommandList.h"
+#include "IMeshPaintGeometryAdapter.h"
+#include "MeshPaintHelpers.h"
+#include "MeshPaintColorBrush.h"
+#include "MeshPaintMode.generated.h"
+
+class UMeshVertexPaintingToolProperties;
+class UMeshColorPaintingToolProperties;
+class UMeshPaintModeSettings;
+class IMeshPaintGeometryAdapter;
+class UMeshComponent;
+
+/**
+ * Mesh paint Mode.  Extends editor viewports with the ability to paint data on meshes
+ */
+UCLASS()
+class UMeshPaintMode : public UEdMode
+{
+public:
+	GENERATED_BODY()
+
+	/** Default constructor for UMeshPaintMode */
+	UMeshPaintMode();
+	static UMeshVertexPaintingToolProperties* GetVertexToolProperties();
+	static UMeshColorPaintingToolProperties* GetColorToolProperties();
+	static UMeshWeightPaintingToolProperties* GetWeightToolProperties();
+	static UMeshPaintMode* GetMeshPaintMode();
+	virtual void Enter() override;
+	virtual void Exit() override;
+	virtual void CreateToolkit() override;
+	virtual void Tick(FEditorViewportClient* ViewportClient, float DeltaTime) override;
+	static FName MeshPaintMode_Color;
+	static FName MeshPaintMode_Texture;
+	static FName MeshPaintMode_Weights;
+
+	virtual TMap<FName, TArray<TSharedPtr<FUICommandInfo>>> GetModeCommands() const override;
+	/** Returns the instance of ComponentClass found in the current Editor selection */
+	template<typename ComponentClass>
+	TArray<ComponentClass*> GetSelectedComponents() const;
+
+	uint32 GetCachedVertexDataSize() const
+	{
+		return CachedVertexDataSize;
+	}
+
+protected:
+
+	/** Binds UI commands to actions for the mesh paint mode */
+	virtual void BindCommands();
+
+	// UEdMode interface
+	virtual void OnToolStarted(UInteractiveToolManager* Manager, UInteractiveTool* Tool) override;
+	virtual void OnToolEnded(UInteractiveToolManager* Manager, UInteractiveTool* Tool) override;
+	virtual bool InputAxis(FEditorViewportClient* InViewportClient, FViewport* Viewport, int32 ControllerId, FKey Key, float Delta, float DeltaTime) override;
+	virtual void ActorSelectionChangeNotify() override;
+	virtual void ActivateDefaultTool() override;
+	virtual void UpdateOnPaletteChange() override;
+
+	// end UEdMode Interface
+	void UpdateSelectedMeshes();
+	void FillWithVertexColor();
+	void PropagateVertexColorsToAsset();
+	bool CanPropagateVertexColors() const;
+	void ImportVertexColors();
+	void SavePaintedAssets();
+	bool CanSaveMeshPackages() const;
+	bool CanRemoveInstanceColors() const;
+	bool CanPasteInstanceVertexColors() const;
+	bool CanCopyInstanceVertexColors() const;
+	bool CanPropagateVertexColorsToLODs() const;
+	/** Copy and pasting functionality from and to a StaticMeshComponent (currently only supports per instance)*/
+	void CopyVertexColors();
+	void PasteVertexColors();
+	void FixVertexColors();
+	bool DoesRequireVertexColorsFixup() const;
+	void RemoveVertexColors();
+	void PropagateVertexColorsToLODs();
+	void UpdateCachedVertexDataSize();
+
+protected:
+	UPROPERTY(Transient)
+	UMeshPaintModeSettings* ModeSettings;
+	// End vertex paint state
+	FGetSelectedMeshComponents MeshComponentDelegate;
+	uint32 CachedVertexDataSize;
+	bool bRecacheVertexDataSize;
+
+};
+
diff --git a/Engine/Plugins/MeshPainting/Source/MeshPaintEditorMode/Private/MeshPaintModeCommands.cpp b/Engine/Plugins/MeshPainting/Source/MeshPaintEditorMode/Private/MeshPaintModeCommands.cpp
new file mode 100644
--- /dev/null
+++ b/Engine/Plugins/MeshPainting/Source/MeshPaintEditorMode/Private/MeshPaintModeCommands.cpp
@@ -0,0 +1,112 @@
+// Copyright 1998-2019 Epic Games, Inc. All Rights Reserved.
+
+#include "MeshPaintModeCommands.h"
+#include "Framework/Commands/UIAction.h"
+#include "MeshPaintMode.h"
+#include "SingleSelectionTool.h"
+#include "MeshPaintColorBrush.h"
+
+#define LOCTEXT_NAMESPACE "MeshPaintEditorModeCommands"
+
+
+
+void FMeshPaintingToolActionCommands::GetToolDefaultObjectList(TArray<UInteractiveTool*>& ToolCDOs)
+{
+	ToolCDOs.Add(GetMutableDefault<USingleSelectionTool>());
+	ToolCDOs.Add(GetMutableDefault<UMeshColorPaintingTool>());
+}
+
+
+
+
+void FMeshPaintingToolActionCommands::RegisterAllToolActions()
+{
+	FMeshPaintingToolActionCommands::Register();
+}
+
+void FMeshPaintingToolActionCommands::UnregisterAllToolActions()
+{
+	FMeshPaintingToolActionCommands::Unregister();
+}
+
+void FMeshPaintingToolActionCommands::UpdateToolCommandBinding(UInteractiveTool* Tool, TSharedPtr<FUICommandList> UICommandList, bool bUnbind)
+{
+	if (FMeshPaintingToolActionCommands::IsRegistered())
+	{
+		!bUnbind ? FMeshPaintingToolActionCommands::Get().BindCommandsForCurrentTool(UICommandList, Tool) : FMeshPaintingToolActionCommands::Get().UnbindActiveCommands(UICommandList);
+	}
+}
+
+
+
+
+void FMeshPaintEditorModeCommands::RegisterCommands()
+{
+	TArray<TSharedPtr<FUICommandInfo>> ColorCommands;
+	TArray<TSharedPtr<FUICommandInfo>> WeightCommands;
+	TArray<TSharedPtr<FUICommandInfo>> VertexCommands;
+	 
+	UI_COMMAND(Select, "Select", "Select the mesh", EUserInterfaceActionType::ToggleButton, FInputChord());
+	ColorCommands.Add(Select);
+	WeightCommands.Add(Select);
+	UI_COMMAND(ColorPaint, "Paint", "Paint the mesh", EUserInterfaceActionType::ToggleButton, FInputChord());
+	ColorCommands.Add(ColorPaint);
+	UI_COMMAND(WeightPaint, "Paint", "Paint the mesh", EUserInterfaceActionType::ToggleButton, FInputChord());
+	WeightCommands.Add(WeightPaint);
+	UI_COMMAND(Fill, "Fill", "Fills the selected Meshes with the Paint Color", EUserInterfaceActionType::Button, FInputChord());
+	VertexCommands.Add(Fill);
+	UI_COMMAND(Propagate, "Apply", "Propagates Instance Vertex Colors to the Source Meshes", EUserInterfaceActionType::Button, FInputChord());
+	VertexCommands.Add(Propagate);
+ 	UI_COMMAND(Import, "Import", "Imports Vertex Colors from a TGA Texture File to the Selected Meshes", EUserInterfaceActionType::Button, FInputChord());
+ 	VertexCommands.Add(Import);
+ 	UI_COMMAND(Save, "Save", "Saves the Source Meshes for the selected Mesh Components", EUserInterfaceActionType::Button, FInputChord());
+	VertexCommands.Add(Save);
+	UI_COMMAND(Copy, "Copy", "Copies Vertex Colors from the selected Mesh Components", EUserInterfaceActionType::Button, FInputChord());
+ 	VertexCommands.Add(Copy);
+	UI_COMMAND(Paste, "Paste", "Tried to Paste Vertex Colors on the selected Mesh Components", EUserInterfaceActionType::Button, FInputChord());
+ 	VertexCommands.Add(Paste);
+	UI_COMMAND(Remove, "Remove", "Removes Vertex Colors from the selected Mesh Components", EUserInterfaceActionType::Button, FInputChord());
+	VertexCommands.Add(Remove);
+	UI_COMMAND(Fix, "Fix", "If necessary fixes Vertex Colors applied to the selected Mesh Components", EUserInterfaceActionType::Button, FInputChord());
+	VertexCommands.Add(Fix);
+	UI_COMMAND(PropagateVertexColorsToLODs, "All LODs", "Applied the Vertex Colors from LOD0 to all LOD levels", EUserInterfaceActionType::Button, FInputChord());
+	VertexCommands.Add(PropagateVertexColorsToLODs);
+	ColorCommands.Append(VertexCommands);
+	WeightCommands.Append(VertexCommands);
+	Commands.Add(UMeshPaintMode::MeshPaintMode_Color, ColorCommands);
+	Commands.Add(UMeshPaintMode::MeshPaintMode_Weights, WeightCommands);
+
+	TArray<TSharedPtr<FUICommandInfo>> TextureCommands;
+	UI_COMMAND(PropagateTexturePaint, "Apply", "Propagates Modifications to the Textures", EUserInterfaceActionType::Button, FInputChord());
+	TextureCommands.Add(PropagateTexturePaint);
+	UI_COMMAND(SaveTexturePaint, "Save", "Saves the Modified Textures for the selected Mesh Components", EUserInterfaceActionType::Button, FInputChord());
+	TextureCommands.Add(SaveTexturePaint);
+	Commands.Add(UMeshPaintMode::MeshPaintMode_Texture, TextureCommands);
+
+	// // 	UI_COMMAND(NextTexture, "Next Texture", "Cycle To Next Texture", EUserInterfaceActionType::Button, FInputChord(EKeys::Period));
+	// // 	Commands.Add(NextTexture);
+	// // 	UI_COMMAND(PreviousTexture, "Previous Texture", "Cycle To Previous Texture", EUserInterfaceActionType::Button, FInputChord(EKeys::Comma));
+	// // 	Commands.Add(PreviousTexture);
+	// // 
+	// // 	UI_COMMAND(CommitTexturePainting, "Commit Texture Painting", "Commits Texture Painting Changes", EUserInterfaceActionType::Button, FInputChord(EKeys::C, EModifierKey::Control | EModifierKey::Shift));
+	// // 	Commands.Add(CommitTexturePainting);
+	// // 
+	// 
+	// 
+	// 
+	// // 
+	// // 	UI_COMMAND(SwitchForeAndBackgroundColor, "Switch Fore and Background Color", "Switches the Fore and Background Colors used for Vertex Painting", EUserInterfaceActionType::None, FInputChord(EKeys::X));
+	// // 	Commands.Add(SwitchForeAndBackgroundColor);
+	// // 
+	// 
+	// // 
+	// // 
+	// // 	UI_COMMAND(CycleToNextLOD, "Cycle to next Mesh LOD", "Cycles to the next possible Mesh LOD to Paint on", EUserInterfaceActionType::None, FInputChord(EKeys::N));
+	// // 	Commands.Add(CycleToNextLOD);
+	// // 
+	// // 	UI_COMMAND(CycleToPreviousLOD, "Cycle to previous Mesh LOD", "Cycles to the previous possible Mesh LOD to Paint on", EUserInterfaceActionType::None, FInputChord(EKeys::B));
+	// // 	Commands.Add(CycleToPreviousLOD);
+}
+
+#undef LOCTEXT_NAMESPACE
+
diff --git a/Engine/Plugins/MeshPainting/Source/MeshPaintEditorMode/Private/MeshPaintModeSettings.cpp b/Engine/Plugins/MeshPainting/Source/MeshPaintEditorMode/Private/MeshPaintModeSettings.cpp
new file mode 100644
--- /dev/null
+++ b/Engine/Plugins/MeshPainting/Source/MeshPaintEditorMode/Private/MeshPaintModeSettings.cpp
@@ -0,0 +1,5 @@
+// Copyright 1998-2019 Epic Games, Inc. All Rights Reserved.
+
+#include "MeshPaintModeSettings.h"
+#include "UnrealEdMisc.h"
+
diff --git a/Engine/Plugins/MeshPainting/Source/MeshPaintEditorMode/Private/MeshPaintModeSettings.h b/Engine/Plugins/MeshPainting/Source/MeshPaintEditorMode/Private/MeshPaintModeSettings.h
new file mode 100644
--- /dev/null
+++ b/Engine/Plugins/MeshPainting/Source/MeshPaintEditorMode/Private/MeshPaintModeSettings.h
@@ -0,0 +1,59 @@
+// Copyright 1998-2019 Epic Games, Inc. All Rights Reserved.
+
+#pragma once
+
+#include "MeshPaintHelpers.h"
+#include "MeshPaintModeSettings.generated.h"
+
+
+/** Mesh paint color view modes (somewhat maps to EVertexColorViewMode engine enum.) */
+UENUM()
+enum class EMeshPaintColorView : uint8
+{
+	/** Normal view mode (vertex color visualization off) */
+	Normal UMETA(DisplayName = "Off"),
+
+	/** RGB only */
+	RGB UMETA(DisplayName = "RGB Channels"),
+
+	/** Alpha only */
+	Alpha UMETA(DisplayName = "Alpha Channel"),
+
+	/** Red only */
+	Red UMETA(DisplayName = "Red Channel"),
+
+	/** Green only */
+	Green UMETA(DisplayName = "Green Channel"),
+
+	/** Blue only */
+	Blue UMETA(DisplayName = "Blue Channel"),
+};
+
+/**
+ * Implements the Mesh Editor's settings.
+ */
+UCLASS(config=EditorPerProjectUserSettings)
+class MESHPAINTEDITORMODE_API UMeshPaintModeSettings
+	: public UObject
+{
+	GENERATED_BODY()
+
+public:
+
+	UMeshPaintModeSettings()
+	{
+		VertexPreviewSize = 6.0f;
+		ColorViewMode = EMeshPaintDataColorViewMode::Normal;
+	}
+
+public:
+	/** Size of vertex points drawn when mesh painting is active. */
+	UPROPERTY(config, EditAnywhere, Category = Visualization)
+	float VertexPreviewSize;
+
+	/** Color view mode used to display Vertex Colors */
+	UPROPERTY(config, EditAnywhere, Category = Visualization)
+	EMeshPaintDataColorViewMode ColorViewMode;
+
+
+};
diff --git a/Engine/Plugins/MeshPainting/Source/MeshPaintEditorMode/Private/MeshPaintModeToolkit.cpp b/Engine/Plugins/MeshPainting/Source/MeshPaintEditorMode/Private/MeshPaintModeToolkit.cpp
new file mode 100644
--- /dev/null
+++ b/Engine/Plugins/MeshPainting/Source/MeshPaintEditorMode/Private/MeshPaintModeToolkit.cpp
@@ -0,0 +1,51 @@
+// Copyright 1998-2019 Epic Games, Inc. All Rights Reserved.
+
+#include "MeshPaintModeToolkit.h"
+#include "MeshPaintMode.h"
+
+#define LOCTEXT_NAMESPACE "MeshPaintModeToolkit"
+
+
+FName FMeshPaintModeToolkit::GetToolkitFName() const
+{
+	return FName( "MeshPaintMode" );
+}
+
+
+FText FMeshPaintModeToolkit::GetBaseToolkitName() const
+{
+	return LOCTEXT( "ToolkitName", "Mesh Paint Mode" );
+}
+
+
+class UEdMode* FMeshPaintModeToolkit::GetScriptableEditorMode() const
+{
+	return GLevelEditorModeTools().GetActiveScriptableMode( "MeshPaintMode" );
+}
+
+void FMeshPaintModeToolkit::GetToolPaletteNames(TArray<FName>& PaletteNames) const
+{
+	PaletteNames.Add(UMeshPaintMode::MeshPaintMode_Color);
+	PaletteNames.Add(UMeshPaintMode::MeshPaintMode_Weights);
+	PaletteNames.Add(UMeshPaintMode::MeshPaintMode_Texture);
+}
+
+
+FText FMeshPaintModeToolkit::GetToolPaletteDisplayName(FName Palette) const
+{
+	if (Palette == UMeshPaintMode::MeshPaintMode_Color)
+	{
+		return LOCTEXT("MeshPaintMode_Color", "Colors");
+	}
+	if (Palette == UMeshPaintMode::MeshPaintMode_Weights)
+	{
+		return LOCTEXT("MeshPaintMode_Weights", "Weights");
+	}
+	if (Palette == UMeshPaintMode::MeshPaintMode_Texture)
+	{
+		return LOCTEXT("MeshPaintMode_Texture", "Textures");
+	}
+	return FText();
+}
+
+#undef LOCTEXT_NAMESPACE
diff --git a/Engine/Plugins/MeshPainting/Source/MeshPaintEditorMode/Public/MeshPaintModeCommands.h b/Engine/Plugins/MeshPainting/Source/MeshPaintEditorMode/Public/MeshPaintModeCommands.h
new file mode 100644
--- /dev/null
+++ b/Engine/Plugins/MeshPainting/Source/MeshPaintEditorMode/Public/MeshPaintModeCommands.h
@@ -0,0 +1,101 @@
+// Copyright 1998-2019 Epic Games, Inc. All Rights Reserved.
+
+#pragma once
+
+#include "Framework/Commands/UICommandInfo.h"
+#include "Framework/Commands/Commands.h"
+#include "EditorStyleSet.h"
+#include "Tools/InteractiveToolsCommands.h"
+
+
+class FMeshPaintEditorModeCommands : public TCommands<FMeshPaintEditorModeCommands>
+{
+public:
+	FMeshPaintEditorModeCommands() : TCommands<FMeshPaintEditorModeCommands>
+		(
+			"MeshPaint",
+			NSLOCTEXT("MeshPaintEditorMode", "MeshPaintingModeCommands", "Mesh Painting Mode"),
+			NAME_None,
+			FEditorStyle::GetStyleSetName()
+			)
+	{}
+
+	/**
+	* Initialize commands
+	*/
+	virtual void RegisterCommands() override;
+
+	static TMap<FName, TArray<TSharedPtr<FUICommandInfo>>> GetCommands()
+	{
+		return FMeshPaintEditorModeCommands::Get().Commands;
+	}
+
+public:
+	TSharedPtr<FUICommandInfo> NextTexture;
+	TSharedPtr<FUICommandInfo> PreviousTexture;
+	TSharedPtr<FUICommandInfo> CommitTexturePainting;
+
+	TSharedPtr<FUICommandInfo> Copy;
+	TSharedPtr<FUICommandInfo> Paste;
+	TSharedPtr<FUICommandInfo> Remove;
+	TSharedPtr<FUICommandInfo> Fix;
+	TSharedPtr<FUICommandInfo> Fill;
+	TSharedPtr<FUICommandInfo> Propagate;
+	TSharedPtr<FUICommandInfo> Import;
+	TSharedPtr<FUICommandInfo> Save;
+
+	TSharedPtr<FUICommandInfo> SwitchForeAndBackgroundColor;
+	TSharedPtr<FUICommandInfo> CycleToNextLOD;
+	TSharedPtr<FUICommandInfo> CycleToPreviousLOD;
+
+	TSharedPtr<FUICommandInfo> PropagateTexturePaint;
+	TSharedPtr<FUICommandInfo> SaveTexturePaint;
+
+	TSharedPtr<FUICommandInfo> PropagateVertexColorsToLODs;
+
+	TSharedPtr<FUICommandInfo> ColorPaint;
+	TSharedPtr<FUICommandInfo> WeightPaint;
+	TSharedPtr<FUICommandInfo> TexturePaint;
+	TSharedPtr<FUICommandInfo> Select;
+
+	TMap<FName, TArray<TSharedPtr<FUICommandInfo>>> Commands;
+};
+
+
+class FMeshPaintingToolActionCommands : public TInteractiveToolCommands<FMeshPaintingToolActionCommands>
+{
+public:
+	FMeshPaintingToolActionCommands() :
+		TInteractiveToolCommands<FMeshPaintingToolActionCommands>(
+			"MeshPaintingTools", // Context name for fast lookup
+			NSLOCTEXT("MeshPaintEditorMode", "MeshPaintingToolsCommands", "Mesh Painting Tools"), // Localized context name for displaying
+			NAME_None, // Parent
+			FEditorStyle::GetStyleSetName() // Icon Style Set
+			)
+	{
+	}
+
+	virtual void GetToolDefaultObjectList(TArray<UInteractiveTool*>& ToolCDOs) override;
+
+
+	/**
+	 * interface that hides various per-tool action sets
+	 */
+
+	 /**
+	  * Register all Tool command sets. Call this in module startup
+	  */
+	static void RegisterAllToolActions();
+
+	/**
+	 * Unregister all Tool command sets. Call this from module shutdown.
+	 */
+	static void UnregisterAllToolActions();
+
+	/**
+	 * Add or remove commands relevant to Tool to the given UICommandList.
+	 * Call this when the active tool changes (eg on ToolManager.OnToolStarted / OnToolEnded)
+	 * @param bUnbind if true, commands are removed, otherwise added
+	 */
+	static void UpdateToolCommandBinding(UInteractiveTool* Tool, TSharedPtr<FUICommandList> UICommandList, bool bUnbind = false);
+};
\ No newline at end of file
diff --git a/Engine/Plugins/MeshPainting/Source/MeshPaintEditorMode/Public/MeshPaintModeHelpers.h b/Engine/Plugins/MeshPainting/Source/MeshPaintEditorMode/Public/MeshPaintModeHelpers.h
new file mode 100644
--- /dev/null
+++ b/Engine/Plugins/MeshPainting/Source/MeshPaintEditorMode/Public/MeshPaintModeHelpers.h
@@ -0,0 +1,72 @@
+// Copyright 1998-2019 Epic Games, Inc. All Rights Reserved.
+
+#pragma once
+
+#include "CoreMinimal.h"
+#include "Engine/EngineTypes.h"
+#include "MeshPaintingToolsetTypes.h"
+#include "Kismet/BlueprintFunctionLibrary.h"
+#include "IMeshPaintGeometryAdapter.h"
+
+class FMeshPaintParameters;
+class UMeshVertexColorImportOptions;
+class UTexture2D;
+class UStaticMeshComponent;
+class UStaticMesh;
+class USkeletalMesh;
+class IMeshPaintGeometryAdapter;
+class UPaintBrushSettings;
+class FEditorViewportClient;
+class UMeshComponent;
+class USkeletalMeshComponent;
+class UViewportInteractor;
+class FViewport;
+class FPrimitiveDrawInterface;
+class FSceneView;
+
+struct FStaticMeshComponentLODInfo;
+struct FPerComponentVertexColorData;
+
+enum class EMeshPaintDataColorViewMode : uint8;
+
+class MESHPAINTEDITORMODE_API UMeshPaintModeHelpers : public UBlueprintFunctionLibrary
+{
+public:
+	/** Forces the Viewport Client to render using the given Viewport Color ViewMode */
+	static void SetViewportColorMode(EMeshPaintDataColorViewMode ColorViewMode, FEditorViewportClient* ViewportClient);
+
+	/** Sets whether or not the level viewport should be real time rendered move or viewport as parameter? */
+	static void SetRealtimeViewport(bool bRealtime);
+
+
+	/** Helper function to import Vertex Colors from a Texture to the specified MeshComponent (makes use of SImportVertexColorsOptions Widget) */
+	static void ImportVertexColorsFromTexture(UMeshComponent* MeshComponent);
+
+	/** Imports vertex colors from a Texture to the specified Skeletal Mesh according to user-set options */
+	static void ImportVertexColorsToSkeletalMesh(USkeletalMesh* SkeletalMesh, const UMeshVertexColorImportOptions* Options, UTexture2D* Texture);
+
+	struct FPaintRay
+	{
+		FVector CameraLocation;
+		FVector RayStart;
+		FVector RayDirection;
+		UViewportInteractor* ViewportInteractor;
+	};
+
+
+	static bool RetrieveViewportPaintRays(const FSceneView* View, FViewport* Viewport, FPrimitiveDrawInterface* PDI, TArray<FPaintRay>& OutPaintRays);
+
+	/** Imports vertex colors from a Texture to the specified Static Mesh according to user-set options */
+	static void ImportVertexColorsToStaticMesh(UStaticMesh* StaticMesh, const UMeshVertexColorImportOptions* Options, UTexture2D* Texture);
+
+	/** Imports vertex colors from a Texture to the specified Static Mesh Component according to user-set options */
+	static void ImportVertexColorsToStaticMeshComponent(UStaticMeshComponent* StaticMeshComponent, const UMeshVertexColorImportOptions* Options, UTexture2D* Texture);
+
+	static void PropagateVertexColors(const TArray<UStaticMeshComponent *> StaticMeshComponents);
+	static bool CanPropagateVertexColors(TArray<UStaticMeshComponent*>& StaticMeshComponents, TArray<UStaticMesh*>& StaticMeshes, int32 NumInstanceVertexColorBytes);
+	static void CopyVertexColors(const TArray<UStaticMeshComponent*> StaticMeshComponents, TArray<FPerComponentVertexColorData>& CopiedVertexColors);
+	static bool CanCopyInstanceVertexColors(const TArray<UStaticMeshComponent*>& StaticMeshComponents, int32 PaintingMeshLODIndex);
+	static void PasteVertexColors(const TArray<UStaticMeshComponent*>& StaticMeshComponents, TArray<FPerComponentVertexColorData>& CopiedColorsByComponent);
+	static bool CanPasteInstanceVertexColors(const TArray<UStaticMeshComponent*>& StaticMeshComponents, const TArray<FPerComponentVertexColorData>& CopiedColorsByComponent);
+	static void RemovePerLODColors(const TArray<UMeshComponent*>& PaintableComponents);
+};
\ No newline at end of file
diff --git a/Engine/Plugins/MeshPainting/Source/MeshPaintingToolset/Private/MeshPaintColorBrush.cpp b/Engine/Plugins/MeshPainting/Source/MeshPaintingToolset/Private/MeshPaintColorBrush.cpp
new file mode 100644
--- /dev/null
+++ b/Engine/Plugins/MeshPainting/Source/MeshPaintingToolset/Private/MeshPaintColorBrush.cpp
@@ -0,0 +1,707 @@
+// Copyright 1998-2019 Epic Games, Inc. All Rights Reserved.
+
+#include "MeshPaintColorBrush.h"
+#include "InteractiveToolManager.h"
+#include "ToolBuilderUtil.h"
+#include "SceneManagement.h" // for FPrimitiveDrawInterface
+#include "ToolContextInterfaces.h"
+#include "Components/MeshComponent.h"
+#include "Math/UnrealMathUtility.h"
+#include "IMeshPaintGeometryAdapter.h"
+#include "ComponentReregisterContext.h"
+#include "MeshPaintAdapterFactory.h"
+#include "ToolDataVisualizer.h"
+//#include "MeshPaintModeSettings.h"
+
+
+#define LOCTEXT_NAMESPACE "MeshVertexBrush"
+
+/*
+ * ToolBuilder
+ */
+
+bool HasPaintableMesh(UActorComponent* Component)
+{
+	return Cast<UMeshComponent>(Component) != nullptr;
+}
+
+
+bool UMeshColorPaintingToolBuilder::CanBuildTool(const FToolBuilderState& SceneState) const
+{
+	UActorComponent* ActorComponent = ToolBuilderUtil::FindFirstComponent(SceneState, HasPaintableMesh);
+	return ActorComponent != nullptr;
+}
+
+UInteractiveTool* UMeshColorPaintingToolBuilder::BuildTool(const FToolBuilderState& SceneState) const
+{
+	UMeshColorPaintingTool* NewTool = NewObject<UMeshColorPaintingTool>(SceneState.ToolManager, "ColorBrushTool");
+	return NewTool;
+}
+
+bool UMeshWeightPaintingToolBuilder::CanBuildTool(const FToolBuilderState& SceneState) const
+{
+	UActorComponent* ActorComponent = ToolBuilderUtil::FindFirstComponent(SceneState, HasPaintableMesh);
+	return ActorComponent != nullptr;
+}
+
+UInteractiveTool* UMeshWeightPaintingToolBuilder::BuildTool(const FToolBuilderState& SceneState) const
+{
+	UMeshWeightPaintingTool* NewTool = NewObject<UMeshWeightPaintingTool>(SceneState.ToolManager, "ColorWeightTool");
+	return NewTool;
+}
+
+
+/*
+ * Tool
+ */
+
+UMeshVertexPaintingToolProperties::UMeshVertexPaintingToolProperties()
+	:UBrushBaseProperties(),
+	PaintColor(FLinearColor::White),
+	EraseColor(FLinearColor::Black)
+{
+}
+
+
+
+void UMeshVertexPaintingToolProperties::SaveProperties(UInteractiveTool* SaveFromTool)
+{
+	UBrushBaseProperties::SaveProperties(SaveFromTool);
+	UMeshVertexPaintingToolProperties* PropertyCache = GetPropertyCache<UMeshVertexPaintingToolProperties>();
+	PropertyCache->PaintColor = this->PaintColor;
+	PropertyCache->EraseColor = this->EraseColor;
+	PropertyCache->bOnlyFrontFacingTriangles = this->bOnlyFrontFacingTriangles;
+}
+
+void UMeshVertexPaintingToolProperties::RestoreProperties(UInteractiveTool* RestoreToTool)
+{
+	UBrushBaseProperties::RestoreProperties(RestoreToTool);
+	UMeshVertexPaintingToolProperties* PropertyCache = GetPropertyCache<UMeshVertexPaintingToolProperties>();
+	this->PaintColor = PropertyCache->PaintColor;
+	this->EraseColor = PropertyCache->EraseColor;
+	this->bOnlyFrontFacingTriangles = PropertyCache->bOnlyFrontFacingTriangles;
+}
+
+UMeshVertexPaintingTool::UMeshVertexPaintingTool()
+{
+	PropertyClass = UMeshVertexPaintingToolProperties::StaticClass();
+}
+
+
+void UMeshVertexPaintingTool::Setup()
+{
+	Super::Setup();
+	VertexProperties = Cast<UMeshVertexPaintingToolProperties>(BrushProperties);
+	bResultValid = false;
+	bStampPending = false;
+	BrushProperties->RestoreProperties(this);
+}
+
+
+void UMeshVertexPaintingTool::Shutdown(EToolShutdownType ShutdownType)
+{
+	FinishPainting();
+	BrushProperties->SaveProperties(this);
+	Super::Shutdown(ShutdownType);
+}
+
+void UMeshVertexPaintingTool::Render(IToolsContextRenderAPI* RenderAPI)
+{
+	UMeshToolManager* MeshToolManager = Cast<UMeshToolManager>(GetToolManager());
+	if (MeshToolManager && LastBestHitResult.Component != nullptr)
+	{
+		Super::Render(RenderAPI);
+		FToolDataVisualizer Draw;
+		Draw.BeginFrame(RenderAPI);
+		static float WidgetLineThickness = 1.0f;
+		static FLinearColor VertexPointColor = FLinearColor::White;
+		static FLinearColor	HoverVertexPointColor = FLinearColor(0.3f, 1.0f, 0.3f);
+		const float NormalLineSize(BrushProperties->BrushRadius * 0.35f);	// Make the normal line length a function of brush size
+		static const FLinearColor NormalLineColor(0.3f, 1.0f, 0.3f);
+		const FLinearColor BrushCueColor = (bArePainting ? FLinearColor(1.0f, 1.0f, 0.3f) : FLinearColor(0.3f, 1.0f, 0.3f));
+ 		const FLinearColor InnerBrushCueColor = (bArePainting ? FLinearColor(0.5f, 0.5f, 0.1f) : FLinearColor(0.1f, 0.5f, 0.1f));
+		const float PointDrawSize = 6.0f;//GetDefault<UMeshPaintModeSettings>()->VertexPreviewSize;
+		// Draw trace surface normal
+		const FVector NormalLineEnd(LastBestHitResult.Location + LastBestHitResult.Normal * NormalLineSize);
+		Draw.DrawLine(FVector(LastBestHitResult.Location), NormalLineEnd, NormalLineColor, WidgetLineThickness);
+
+		for (UMeshComponent* CurrentComponent : MeshToolManager->GetPaintableMeshComponents())
+		{
+			TSharedPtr<IMeshPaintGeometryAdapter> MeshAdapter = MeshToolManager->GetAdapterForComponent(Cast<UMeshComponent>(CurrentComponent));
+
+			if (MeshAdapter->IsValid() /*&& bRenderVertices*/ && MeshAdapter->SupportsVertexPaint())
+			{
+				const FMatrix ComponentToWorldMatrix = MeshAdapter->GetComponentToWorldMatrix();
+				FViewCameraState CameraState;
+				GetToolManager()->GetContextQueriesAPI()->GetCurrentViewState(CameraState);
+				const FVector ComponentSpaceCameraPosition(ComponentToWorldMatrix.InverseTransformPosition(CameraState.Position));
+				const FVector ComponentSpaceBrushPosition(ComponentToWorldMatrix.InverseTransformPosition(LastBestHitResult.Location));
+
+				// @todo MeshPaint: Input vector doesn't work well with non-uniform scale
+				const float ComponentSpaceBrushRadius = ComponentToWorldMatrix.InverseTransformVector(FVector(BrushProperties->BrushRadius, 0.0f, 0.0f)).Size();
+				const float ComponentSpaceSquaredBrushRadius = ComponentSpaceBrushRadius * ComponentSpaceBrushRadius;
+
+				const TArray<FVector>& InRangeVertices = MeshAdapter->SphereIntersectVertices(ComponentSpaceSquaredBrushRadius, ComponentSpaceBrushPosition, ComponentSpaceCameraPosition, VertexProperties->bOnlyFrontFacingTriangles);
+
+				for (const FVector& Vertex : InRangeVertices)
+				{
+					const FVector WorldPositionVertex = ComponentToWorldMatrix.TransformPosition(Vertex);
+					if ((LastBestHitResult.Location - WorldPositionVertex).Size() <= BrushProperties->BrushRadius)
+					{
+						const float VisualBiasDistance = 0.15f;
+						const FVector VertexVisualPosition = WorldPositionVertex + LastBestHitResult.Normal * VisualBiasDistance;
+						Draw.DrawPoint(VertexVisualPosition, HoverVertexPointColor, PointDrawSize, SDPG_World);
+					}
+				}
+			}
+		}
+		Draw.EndFrame();
+	}
+	UpdateResult();
+
+}
+
+void UMeshVertexPaintingTool::Tick(float DeltaTime)
+{
+	Super::Tick(DeltaTime);
+
+	if (UMeshToolManager* MeshToolManager = Cast<UMeshToolManager>(GetToolManager()))
+	{
+		if (MeshToolManager->bNeedsRecache)
+		{
+			CacheSelectionData();
+			bDoRestoreRenTargets = true;
+		}
+	}
+
+	if (bStampPending)
+	{
+		Paint(PendingStampRay.Origin, PendingStampRay.Direction);
+		bStampPending = false;
+
+		// flow
+		if (bInDrag)
+		{
+			bStampPending = true;
+		}
+	}
+}
+
+void UMeshVertexPaintingTool::OnPropertyModified(UObject* PropertySet, UProperty* Property)
+{
+	bResultValid = false;
+}
+
+
+double UMeshVertexPaintingTool::CalculateTargetEdgeLength(int TargetTriCount)
+{
+	double TargetTriArea = InitialMeshArea / (double)TargetTriCount;
+	double EdgeLen = (TargetTriArea);
+	return (double)FMath::RoundToInt(EdgeLen*100.0) / 100.0;
+}
+
+bool UMeshVertexPaintingTool::Paint(const FVector& InRayOrigin, const FVector& InRayDirection)
+{
+	// Determine paint action according to whether or not shift is held down
+	const EMeshPaintModeAction PaintAction = GetShiftToggle() ? EMeshPaintModeAction::Erase : EMeshPaintModeAction::Paint;
+	const float PaintStrength = 1.0f; //Viewport->IsPenActive() ? Viewport->GetTabletPressure() : 1.f;
+	// Handle internal painting functionality
+	TPair<FVector, FVector> Ray(InRayOrigin, InRayDirection);
+	return PaintInternal(MakeArrayView(&Ray, 1), PaintAction, PaintStrength);
+}
+
+bool UMeshVertexPaintingTool::Paint(const TArrayView<TPair<FVector, FVector>>& Rays)
+{
+	// Determine paint action according to whether or not shift is held down
+	const EMeshPaintModeAction PaintAction = GetShiftToggle() ? EMeshPaintModeAction::Erase : EMeshPaintModeAction::Paint;
+
+	const float PaintStrength = 1.0f; //Viewport->IsPenActive() ? Viewport->GetTabletPressure() : 1.f;
+	// Handle internal painting functionality
+	return PaintInternal(Rays, PaintAction, PaintStrength);
+}
+
+
+bool UMeshVertexPaintingTool::PaintInternal(const TArrayView<TPair<FVector, FVector>>& Rays, EMeshPaintModeAction PaintAction, float PaintStrength)
+{
+	TArray<FPaintRayResults> PaintRayResults;
+	PaintRayResults.AddDefaulted(Rays.Num());
+	LastBestHitResult.Reset();
+	TMap<UMeshComponent*, TArray<int32>> HoveredComponents;
+
+	const float BrushRadius = BrushProperties->BrushRadius;
+	const bool bIsPainting = PaintAction == EMeshPaintModeAction::Paint;
+	const float InStrengthScale = PaintStrength;
+
+	bool bPaintApplied = false;
+	UMeshToolManager* MeshToolManager = Cast<UMeshToolManager>(GetToolManager());
+	// Fire out a ray to see if there is a *selected* component under the mouse cursor that can be painted.
+	for (int32 i = 0; i < Rays.Num(); ++i)
+	{
+		const FVector& RayOrigin = Rays[i].Key;
+		const FVector& RayDirection = Rays[i].Value;
+		FRay Ray = FRay(RayOrigin, RayDirection);
+		FHitResult& BestTraceResult = PaintRayResults[i].BestTraceResult;
+		const FVector TraceStart(RayOrigin);
+		const FVector TraceEnd(RayOrigin + RayDirection * HALF_WORLD_MAX);
+
+		for (UMeshComponent* MeshComponent : MeshToolManager->GetPaintableMeshComponents())
+		{
+			TSharedPtr<IMeshPaintGeometryAdapter> MeshAdapter = MeshToolManager->GetAdapterForComponent(MeshComponent);
+
+			// Ray trace
+			FHitResult TraceHitResult(1.0f);
+
+			if (MeshAdapter->LineTraceComponent(TraceHitResult, TraceStart, TraceEnd, FCollisionQueryParams(SCENE_QUERY_STAT(Paint), true)))
+			{
+				// Find the closest impact
+				if ((BestTraceResult.GetComponent() == nullptr) || (TraceHitResult.Time < BestTraceResult.Time))
+				{
+					BestTraceResult = TraceHitResult;
+				}
+			}
+		}
+
+		bool bUsed = false;
+
+		if (BestTraceResult.GetComponent() != nullptr)
+		{
+			FBox BrushBounds = FBox::BuildAABB(BestTraceResult.Location, FVector(BrushRadius * 1.25f, BrushRadius * 1.25f, BrushRadius * 1.25f));
+
+			// Vertex paint mode, so we want all valid components overlapping the brush hit location
+			for (auto TestComponent : MeshToolManager->GetPaintableMeshComponents())
+			{
+				const FBox ComponentBounds = TestComponent->Bounds.GetBox();
+
+				if (MeshToolManager->GetComponentToAdapterMap().Contains(TestComponent) && ComponentBounds.Intersect(BrushBounds))
+				{
+					// OK, this mesh potentially overlaps the brush!
+					HoveredComponents.FindOrAdd(TestComponent).Add(i);
+					bUsed = true;
+				}
+			}
+		}
+		if (bUsed)
+		{
+			FVector BrushXAxis, BrushYAxis;
+			BestTraceResult.Normal.FindBestAxisVectors(BrushXAxis, BrushYAxis);
+			// Display settings
+			const float VisualBiasDistance = 0.15f;
+			const FVector BrushVisualPosition = BestTraceResult.Location + BestTraceResult.Normal * VisualBiasDistance;
+
+			const FLinearColor PaintColor = VertexProperties->PaintColor;
+			const FLinearColor EraseColor = VertexProperties->EraseColor;
+
+			// NOTE: We square the brush strength to maximize slider precision in the low range
+			const float BrushStrength = BrushProperties->BrushStrength *  BrushProperties->BrushStrength * InStrengthScale;
+
+			const float BrushDepth = BrushRadius;
+			LastBestHitResult = BestTraceResult;
+			// Mesh paint settings
+			FMeshPaintParameters& Params = PaintRayResults[i].Params;
+			Params.PaintAction = PaintAction;
+			Params.BrushPosition = BestTraceResult.Location;
+			Params.BrushNormal = BestTraceResult.Normal;
+			Params.BrushColor = bIsPainting ? PaintColor : EraseColor;
+			Params.SquaredBrushRadius = BrushRadius * BrushRadius;
+			Params.BrushRadialFalloffRange = BrushProperties->BrushFalloffAmount * BrushRadius;
+			Params.InnerBrushRadius = BrushRadius - Params.BrushRadialFalloffRange;
+			Params.BrushDepth = BrushDepth;
+			Params.BrushDepthFalloffRange = BrushProperties->BrushFalloffAmount * BrushDepth;
+			Params.BrushDepthFalloffRange = BrushProperties->BrushFalloffAmount * BrushDepth;
+			Params.InnerBrushDepth = BrushDepth - Params.BrushDepthFalloffRange;
+			Params.BrushStrength = BrushStrength;
+			Params.BrushToWorldMatrix = FMatrix(BrushXAxis, BrushYAxis, Params.BrushNormal, Params.BrushPosition);
+			Params.InverseBrushToWorldMatrix = Params.BrushToWorldMatrix.Inverse();
+
+			SetAdditionalPaintParameters(Params);
+		
+		}
+	}
+
+	if (HoveredComponents.Num() > 0)
+	{
+		if (bArePainting == false)
+		{
+				// Vertex painting is an ongoing transaction, while texture painting is handled separately later in a single transaction
+				GetToolManager()->BeginUndoTransaction(LOCTEXT("MeshPaintMode_VertexPaint_TransactionPaintStroke", "Vertex Paint"));
+				bArePainting = true;
+				TimeSinceStartedPainting = 0.0f;
+		}
+
+		// Iterate over the selected meshes under the cursor and paint them!
+		for (auto& Entry : HoveredComponents)
+		{
+			UMeshComponent* HoveredComponent = Entry.Key;
+			TArray<int32>& PaintRayResultIds = Entry.Value;
+
+			IMeshPaintGeometryAdapter* MeshAdapter = MeshToolManager->GetAdapterForComponent(HoveredComponent).Get();
+			if (!ensure(MeshAdapter))
+			{
+				continue;
+			}
+
+			if (MeshAdapter->SupportsVertexPaint())
+			{
+				FPerVertexPaintActionArgs Args;
+				Args.Adapter = MeshAdapter;
+				FViewCameraState CameraState;
+				GetToolManager()->GetContextQueriesAPI()->GetCurrentViewState(CameraState);
+				Args.CameraPosition = CameraState.Position;
+				Args.BrushProperties = VertexProperties;
+				Args.Action = PaintAction;
+
+				bool bMeshPreEditCalled = false;
+
+				TSet<int32> InfluencedVertices;
+				for (int32 PaintRayResultId : PaintRayResultIds)
+				{
+					InfluencedVertices.Reset();
+					Args.HitResult = PaintRayResults[PaintRayResultId].BestTraceResult;
+					bPaintApplied |= UMeshPaintingToolset::GetPerVertexPaintInfluencedVertices(Args, InfluencedVertices);
+
+					if (InfluencedVertices.Num() == 0)
+					{
+						continue;
+					}
+
+					if (!bMeshPreEditCalled)
+					{
+						bMeshPreEditCalled = true;
+						MeshAdapter->PreEdit();
+					}
+
+					for (const int32 VertexIndex : InfluencedVertices)
+					{
+						ApplyVertexData(Args, VertexIndex, PaintRayResults[PaintRayResultId].Params);
+					}
+				}
+
+				if (bMeshPreEditCalled)
+				{
+					MeshAdapter->PostEdit();
+				}
+			}
+		}
+	}
+
+	return bPaintApplied;
+}
+
+void UMeshVertexPaintingTool::ApplyVertexData(FPerVertexPaintActionArgs& InArgs, int32 VertexIndex, FMeshPaintParameters Parameters)
+{
+	/** Retrieve vertex position and color for applying vertex painting */
+	FColor PaintColor;
+	FVector Position;
+	InArgs.Adapter->GetVertexPosition(VertexIndex, Position);
+	Position = InArgs.Adapter->GetComponentToWorldMatrix().TransformPosition(Position);
+	InArgs.Adapter->GetVertexColor(VertexIndex, PaintColor, true);
+	UMeshPaintingToolset::PaintVertex(Position, Parameters, PaintColor);
+	InArgs.Adapter->SetVertexColor(VertexIndex, PaintColor, true);
+}
+
+
+void UMeshVertexPaintingTool::UpdateResult()
+{
+	GetToolManager()->PostInvalidation();
+
+	bResultValid = true;
+}
+
+bool UMeshVertexPaintingTool::HasAccept() const
+{
+	return false;
+}
+
+bool UMeshVertexPaintingTool::CanAccept() const
+{
+	return false;
+}
+
+void UMeshVertexPaintingTool::OnBeginDrag(const FRay& Ray)
+{
+	Super::OnBeginDrag(Ray);
+	FHitResult OutHit;
+	if (HitTest(Ray, OutHit))
+	{
+		bInDrag = true;
+
+		// apply initial stamp
+		PendingStampRay = Ray;
+		bStampPending = true;
+	}
+}
+
+void UMeshVertexPaintingTool::OnUpdateDrag(const FRay& Ray)
+{
+	Super::OnUpdateDrag(Ray);
+	if (bInDrag)
+	{
+		PendingStampRay = Ray;
+		bStampPending = true;
+	}
+}
+
+
+
+void UMeshVertexPaintingTool::OnEndDrag(const FRay& Ray)
+{
+	FinishPainting();
+	bStampPending = false;
+	bInDrag = false;
+}
+
+bool UMeshVertexPaintingTool::HitTest(const FRay& Ray, FHitResult& OutHit)
+{
+	bool bUsed = false;
+	if (UMeshToolManager* MeshToolManager = Cast<UMeshToolManager>(GetToolManager()))
+	{
+		bUsed = MeshToolManager->FindHitResult(Ray, OutHit);
+		LastBestHitResult = OutHit;
+	}
+	return bUsed;
+}
+
+void UMeshVertexPaintingTool::FinishPainting()
+{
+	if (bArePainting)
+	{
+		bArePainting = false;
+		GetToolManager()->EndUndoTransaction();
+		OnPaintingFinishedDelegate.ExecuteIfBound();
+	}
+}
+
+
+UMeshColorPaintingToolProperties::UMeshColorPaintingToolProperties()
+	:UMeshVertexPaintingToolProperties(),
+	bWriteRed(true),
+	bWriteGreen(true),
+	bWriteBlue(true),
+	bWriteAlpha(false)
+{
+
+}
+
+void UMeshColorPaintingToolProperties::SaveProperties(UInteractiveTool* SaveFromTool)
+{
+	Super::SaveProperties(SaveFromTool);
+	UMeshColorPaintingToolProperties* PropertyCache = GetPropertyCache<UMeshColorPaintingToolProperties>();
+	PropertyCache->bWriteRed = this->bWriteRed;
+	PropertyCache->bWriteGreen = this->bWriteGreen;
+	PropertyCache->bWriteBlue = this->bWriteBlue;
+	PropertyCache->bWriteRed = this->bWriteRed;
+}
+
+void UMeshColorPaintingToolProperties::RestoreProperties(UInteractiveTool* RestoreToTool)
+{
+	Super::RestoreProperties(RestoreToTool);
+	UMeshColorPaintingToolProperties* PropertyCache = GetPropertyCache<UMeshColorPaintingToolProperties>();
+	this->bWriteRed = PropertyCache->bWriteRed;
+	this->bWriteGreen = PropertyCache->bWriteGreen;
+	this->bWriteBlue = PropertyCache->bWriteBlue;
+	this->bWriteRed = PropertyCache->bWriteRed;
+}
+
+UMeshColorPaintingTool::UMeshColorPaintingTool()
+{
+	PropertyClass = UMeshColorPaintingToolProperties::StaticClass();
+}
+
+void UMeshColorPaintingTool::Setup()
+{
+	Super::Setup();
+	ColorProperties = Cast<UMeshColorPaintingToolProperties>(BrushProperties);
+}
+
+void UMeshColorPaintingTool::CacheSelectionData()
+{
+	Super::CacheSelectionData();
+	if (UMeshToolManager* MeshToolManager = Cast<UMeshToolManager>(GetToolManager()))
+	{
+		MeshToolManager->ClearPaintableMeshComponents();
+		// Update(cached) Paint LOD level if necessary
+		ColorProperties->LODIndex = FMath::Min<int32>(ColorProperties->LODIndex, GetMaxLODIndexToPaint());
+		CachedLODIndex = ColorProperties->LODIndex;
+		bCachedForceLOD = ColorProperties->bPaintOnSpecificLOD;
+		//Determine LOD level to use for painting(can only paint on LODs in vertex mode)
+		const int32 PaintLODIndex = ColorProperties->bPaintOnSpecificLOD ? ColorProperties->LODIndex : 0;
+		//Determine UV channel to use while painting textures
+		const int32 UVChannel = 0;
+
+		MeshToolManager->CacheSelectionData(PaintLODIndex, UVChannel);
+	}
+}
+
+void UMeshColorPaintingTool::SetAdditionalPaintParameters(FMeshPaintParameters& InPaintParameters)
+{
+	InPaintParameters.bWriteRed = ColorProperties->bWriteRed;
+	InPaintParameters.bWriteGreen = ColorProperties->bWriteGreen;
+	InPaintParameters.bWriteBlue = ColorProperties->bWriteBlue;
+	InPaintParameters.bWriteAlpha = ColorProperties->bWriteAlpha;
+	InPaintParameters.ApplyVertexDataDelegate.AddStatic(&UMeshPaintingToolset::ApplyVertexColorPaint);
+}
+
+int32 UMeshColorPaintingTool::GetMaxLODIndexToPaint() const
+{
+	//The maximum LOD we can paint is decide by the lowest number of LOD in the selection
+	int32 LODMin = TNumericLimits<int32>::Max();
+
+	if (UMeshToolManager* MeshToolManager = Cast<UMeshToolManager>(GetToolManager()))
+	{
+		TArray<UMeshComponent*> SelectedComponents = MeshToolManager->GetSelectedMeshComponents();
+
+		for (UMeshComponent* MeshComponent : SelectedComponents)
+		{
+			int32 NumMeshLODs = 0;
+			if (UMeshPaintingToolset::TryGetNumberOfLODs(MeshComponent, NumMeshLODs))
+			{
+				ensure(NumMeshLODs > 0);
+				LODMin = FMath::Min(LODMin, NumMeshLODs - 1);
+			}
+		}
+		if (LODMin == TNumericLimits<int32>::Max())
+		{
+			LODMin = 1;
+		}
+	}
+	return LODMin;
+}
+
+
+void UMeshColorPaintingTool::LODPaintStateChanged(const bool bLODPaintingEnabled)
+{
+	bool AbortChange = false;
+
+	// Set actual flag in the settings struct
+	ColorProperties->bPaintOnSpecificLOD = bLODPaintingEnabled;
+
+	if (!bLODPaintingEnabled)
+	{
+		// Reset painting LOD index
+		ColorProperties->LODIndex = 0;
+
+	}
+
+	ApplyForcedLODIndex(bLODPaintingEnabled ? CachedLODIndex : -1);
+
+	if (UMeshToolManager* MeshToolManager = Cast<UMeshToolManager>(GetToolManager()))
+	{
+		TArray<UMeshComponent*> PaintableComponents = MeshToolManager->GetPaintableMeshComponents();
+
+		TUniquePtr< FComponentReregisterContext > ComponentReregisterContext;
+		//Make sure all static mesh render is dirty since we change the force LOD
+		for (UMeshComponent* SelectedComponent : PaintableComponents)
+		{
+			if (SelectedComponent)
+			{
+				ComponentReregisterContext.Reset(new FComponentReregisterContext(SelectedComponent));
+			}
+		}
+
+		MeshToolManager->Refresh();
+	}
+}
+
+void UMeshColorPaintingTool::ApplyForcedLODIndex(int32 ForcedLODIndex)
+{
+	if (UMeshToolManager* MeshToolManager = Cast<UMeshToolManager>(GetToolManager()))
+	{
+		TArray<UMeshComponent*> PaintableComponents = MeshToolManager->GetPaintableMeshComponents();
+
+		for (UMeshComponent* SelectedComponent : PaintableComponents)
+		{
+			if (SelectedComponent)
+			{
+				UMeshPaintingToolset::ForceRenderMeshLOD(SelectedComponent, ForcedLODIndex);
+			}
+		}
+	}
+}
+
+
+void UMeshColorPaintingTool::PaintLODChanged()
+{
+	// Enforced LOD for painting
+	if (CachedLODIndex != ColorProperties->LODIndex)
+	{
+		CachedLODIndex = ColorProperties->LODIndex;
+		ApplyForcedLODIndex(bCachedForceLOD ? CachedLODIndex : -1);
+
+		TUniquePtr< FComponentReregisterContext > ComponentReregisterContext;
+		//Make sure all static mesh render is dirty since we change the force LOD
+		if (UMeshToolManager* MeshToolManager = Cast<UMeshToolManager>(GetToolManager()))
+		{
+			TArray<UMeshComponent*> PaintableComponents = MeshToolManager->GetPaintableMeshComponents();
+
+			for (UMeshComponent* SelectedComponent : PaintableComponents)
+			{
+				if (SelectedComponent)
+				{
+					ComponentReregisterContext.Reset(new FComponentReregisterContext(SelectedComponent));
+				}
+			}
+
+			MeshToolManager->Refresh();
+		}
+	}
+}
+
+UMeshWeightPaintingToolProperties::UMeshWeightPaintingToolProperties()
+	:UMeshVertexPaintingToolProperties(),
+	TextureWeightType(EMeshPaintWeightTypes::AlphaLerp),
+	PaintTextureWeightIndex(EMeshPaintTextureIndex::TextureOne),
+	EraseTextureWeightIndex(EMeshPaintTextureIndex::TextureTwo)
+{
+
+}
+
+void UMeshWeightPaintingToolProperties::SaveProperties(UInteractiveTool* SaveFromTool)
+{
+	Super::SaveProperties(SaveFromTool);
+}
+
+void UMeshWeightPaintingToolProperties::RestoreProperties(UInteractiveTool* RestoreToTool)
+{
+	Super::RestoreProperties(RestoreToTool);
+}
+
+UMeshWeightPaintingTool::UMeshWeightPaintingTool()
+{
+	PropertyClass = UMeshWeightPaintingToolProperties::StaticClass();
+}
+
+void UMeshWeightPaintingTool::Setup()
+{
+	Super::Setup();
+	WeightProperties = Cast<UMeshWeightPaintingToolProperties>(BrushProperties);
+}
+
+void UMeshWeightPaintingTool::CacheSelectionData()
+{
+	Super::CacheSelectionData();
+	if (UMeshToolManager* MeshToolManager = Cast<UMeshToolManager>(GetToolManager()))
+	{
+		MeshToolManager->ClearPaintableMeshComponents();
+		//Determine LOD level to use for painting(can only paint on LODs in vertex mode)
+		const int32 PaintLODIndex = 0;
+		//Determine UV channel to use while painting textures
+		const int32 UVChannel = 0;
+		MeshToolManager->CacheSelectionData(PaintLODIndex, UVChannel);
+	}
+}
+
+void UMeshWeightPaintingTool::SetAdditionalPaintParameters(FMeshPaintParameters& InPaintParameters)
+{
+	InPaintParameters.TotalWeightCount = (int32)WeightProperties->TextureWeightType;
+
+	// Select texture weight index based on whether or not we're painting or erasing
+	{
+		const int32 PaintWeightIndex = InPaintParameters.PaintAction == EMeshPaintModeAction::Paint ? (int32)WeightProperties->PaintTextureWeightIndex : (int32)WeightProperties->EraseTextureWeightIndex;
+
+		// Clamp the weight index to fall within the total weight count
+		InPaintParameters.PaintWeightIndex = FMath::Clamp(PaintWeightIndex, 0, InPaintParameters.TotalWeightCount - 1);
+	}
+	InPaintParameters.ApplyVertexDataDelegate.AddStatic(&UMeshPaintingToolset::ApplyVertexWeightPaint);
+}
+
+#undef LOCTEXT_NAMESPACE
\ No newline at end of file
diff --git a/Engine/Plugins/MeshPainting/Source/MeshPaintingToolset/Private/MeshSelect.cpp b/Engine/Plugins/MeshPainting/Source/MeshPaintingToolset/Private/MeshSelect.cpp
new file mode 100644
--- /dev/null
+++ b/Engine/Plugins/MeshPainting/Source/MeshPaintingToolset/Private/MeshSelect.cpp
@@ -0,0 +1,121 @@
+// Copyright 1998-2019 Epic Games, Inc. All Rights Reserved.
+
+
+#include "MeshSelect.h"
+#include "InputState.h"
+#include "InteractiveToolManager.h"
+#include "ToolContextInterfaces.h"
+#include "MeshPaintHelpers.h"
+#include "Components/MeshComponent.h"
+#include "IMeshPaintGeometryAdapter.h"
+#include "MeshPaintAdapterFactory.h"
+#include "EngineUtils.h"
+#include "Editor.h"
+
+#define LOCTEXT_NAMESPACE "MeshSelection"
+
+
+bool UMeshClickToolBuilder::CanBuildTool(const FToolBuilderState& SceneState) const
+{
+	return true;
+}
+
+UInteractiveTool* UMeshClickToolBuilder::BuildTool(const FToolBuilderState& SceneState) const
+{
+	UMeshClickTool* NewTool = NewObject<UMeshClickTool>(SceneState.ToolManager);
+
+	return NewTool;
+}
+
+
+
+
+UMeshClickTool::UMeshClickTool()
+{
+}
+
+void UMeshClickTool::Setup()
+{
+	UInteractiveTool::Setup();
+
+	// add default button input behaviors for devices
+	USingleClickInputBehavior* MouseBehavior = NewObject<USingleClickInputBehavior>();
+	MouseBehavior->Initialize(this);
+	MouseBehavior->Modifiers.RegisterModifier(AdditiveSelectionModifier, FInputDeviceState::IsShiftKeyDown);
+	AddInputBehavior(MouseBehavior);
+}
+
+void UMeshClickTool::OnUpdateModifierState(int ModifierID, bool bIsOn)
+{
+	if (ModifierID == AdditiveSelectionModifier)
+	{
+		bAddToSelectionSet = bIsOn;
+	}
+}
+
+bool UMeshClickTool::IsHitByClick(const FInputDeviceRay& ClickPos)
+{
+	return true;
+}
+
+void UMeshClickTool::OnClicked(const FInputDeviceRay& ClickPos)
+{
+	if (UMeshToolManager* MeshToolManager = Cast<UMeshToolManager>(GetToolManager()))
+	{
+		const FVector& RayOrigin = ClickPos.WorldRay.Origin;
+		const FVector& RayDirection = ClickPos.WorldRay.Direction;
+		FRay Ray = FRay(RayOrigin, RayDirection);
+		const FVector TraceStart(RayOrigin);
+		const FVector TraceEnd(RayOrigin + RayDirection * HALF_WORLD_MAX);
+		MeshToolManager->ResetState();
+		//Default iterator only iterates over active levels.
+		const EActorIteratorFlags Flags = EActorIteratorFlags::SkipPendingKill;
+		// TODO: The tool needs to know the world its acting on
+		for (TActorIterator<AActor> It(GEditor->GetEditorWorldContext(false).World(), AActor::StaticClass(), Flags); It; ++It)
+		{
+			AActor* Actor = *It;
+			if (Actor->IsEditable() &&
+				Actor->IsListedInSceneOutliner() &&					// Only add actors that are allowed to be selected and drawn in editor
+				!Actor->IsTemplate() &&								// Should never happen, but we never want CDOs
+				!Actor->HasAnyFlags(RF_Transient))				// Don't add transient actors in non-play worlds
+			{
+				TArray<UActorComponent*> CandidateComponents = Actor->K2_GetComponentsByClass(UMeshComponent::StaticClass());
+				for (UActorComponent* CandidateComponent : CandidateComponents)
+				{
+					UMeshComponent* MeshComponent = Cast<UMeshComponent>(CandidateComponent);
+					TSharedPtr<IMeshPaintGeometryAdapter> MeshAdapter = FMeshPaintAdapterFactory::CreateAdapterForMesh(MeshComponent, 0);
+					MeshToolManager->AddPaintableMeshComponent(MeshComponent);
+					MeshToolManager->AddToComponentToAdapterMap(MeshComponent, MeshAdapter);
+				}
+					
+			}
+		}
+		FHitResult TraceHitResult(1.0f);
+		MeshToolManager->FindHitResult(Ray, TraceHitResult);
+		MeshToolManager->ResetState();
+		TArray<AActor*> SelectedActors;
+		if (UMeshComponent* MeshComponent = Cast<UMeshComponent>(TraceHitResult.GetComponent()))
+		{
+			TSharedPtr<IMeshPaintGeometryAdapter> MeshAdapter = FMeshPaintAdapterFactory::CreateAdapterForMesh(MeshComponent, 0);
+			if (MeshComponent->IsVisible() && MeshAdapter.IsValid() && MeshAdapter->IsValid())
+			{
+				MeshToolManager->AddPaintableMeshComponent(MeshComponent);
+				MeshToolManager->AddToComponentToAdapterMap(MeshComponent, MeshAdapter);
+				SelectedActors.Add(Cast<AActor>(MeshComponent->GetOuter()));
+				MeshAdapter->OnAdded();
+			}
+		}
+
+		GetToolManager()->BeginUndoTransaction(LOCTEXT("MeshSelection", "Select Mesh"));
+
+
+		FSelectedOjectsChangeList NewSelection;
+		// TODO add CTRL handling
+		NewSelection.ModificationType = bAddToSelectionSet ? ESelectedObjectsModificationType::Add : ESelectedObjectsModificationType::Replace;
+		NewSelection.Actors.Append(SelectedActors);
+		GetToolManager()->RequestSelectionChange(NewSelection);
+		GetToolManager()->EndUndoTransaction();
+	}
+}
+
+#undef LOCTEXT_NAMESPACE
