First pass on refactoring mesh paint to a plugin that uses the interactive tools framework

diff --git a/Engine/Plugins/MeshPainting/Source/MeshPaintingToolset/Public/BaseMeshPaintGeometryAdapter.h b/Engine/Plugins/MeshPainting/Source/MeshPaintingToolset/Public/BaseMeshPaintGeometryAdapter.h
new file mode 100644
--- /dev/null
+++ b/Engine/Plugins/MeshPainting/Source/MeshPaintingToolset/Public/BaseMeshPaintGeometryAdapter.h
@@ -0,0 +1,40 @@
+// Copyright 1998-2019 Epic Games, Inc. All Rights Reserved.
+
+#pragma once
+
+#include "IMeshPaintGeometryAdapter.h"
+#include "TMeshPaintOctree.h"
+#include "MeshAdapter.h"
+#include "Spatial/MeshAABBTree3.h"
+
+struct FIndex3i;
+typedef TIndexMeshArrayAdapter<uint32, double> FIndexMeshArrayAdapterd;
+
+/** Base mesh paint geometry adapter, handles basic sphere intersection using a Octree */
+class MESHPAINTINGTOOLSET_API FBaseMeshPaintGeometryAdapter : public IMeshPaintGeometryAdapter
+{
+public:
+	/** Start IMeshPaintGeometryAdapter Overrides */
+	virtual bool Initialize() override;
+	virtual const TArray<FVector>& GetMeshVertices() const override;
+	virtual const TArray<uint32>& GetMeshIndices() const override;
+	virtual void GetVertexPosition(int32 VertexIndex, FVector& OutVertex) const override;
+	virtual TArray<uint32> SphereIntersectTriangles(const float ComponentSpaceSquaredBrushRadius, const FVector& ComponentSpaceBrushPosition, const FVector& ComponentSpaceCameraPosition, const bool bOnlyFrontFacing) const override;
+	virtual void GetInfluencedVertexIndices(const float ComponentSpaceSquaredBrushRadius, const FVector& ComponentSpaceBrushPosition, const FVector& ComponentSpaceCameraPosition, const bool bOnlyFrontFacing, TSet<int32> &InfluencedVertices) const override;
+	virtual void GetInfluencedVertexData(const float ComponentSpaceSquaredBrushRadius, const FVector& ComponentSpaceBrushPosition, const FVector& ComponentSpaceCameraPosition, const bool bOnlyFrontFacing, TArray<TPair<int32, FVector>>& OutData) const override;
+	virtual TArray<FVector> SphereIntersectVertices(const float ComponentSpaceSquaredBrushRadius, const FVector& ComponentSpaceBrushPosition, const FVector& ComponentSpaceCameraPosition, const bool bOnlyFrontFacing) const override;
+	virtual bool RayIntersectAdapter(FIndex3i& HitTriangle, FVector& HitPosition, const FVector Start, const FVector End) const override;
+	/** End IMeshPaintGeometryAdapter Overrides */
+
+	virtual bool InitializeVertexData() = 0;
+protected:
+	bool BuildOctree();
+protected:
+	/** Index and Vertex data populated by derived classes in InitializeVertexData */
+	TArray<FVector> MeshVertices;
+	TArray<uint32> MeshIndices;
+	/** Octree used for reducing the cost of sphere intersecting with triangles / vertices */
+	TUniquePtr<FMeshPaintTriangleOctree> MeshTriOctree;
+	FIndexMeshArrayAdapterd Adapter;
+	TUniquePtr<TMeshAABBTree3<const FIndexMeshArrayAdapterd>> AABBTree;
+};
diff --git a/Engine/Plugins/MeshPainting/Source/MeshPaintingToolset/Public/MeshPaintAdapterFactory.h b/Engine/Plugins/MeshPainting/Source/MeshPaintingToolset/Public/MeshPaintAdapterFactory.h
new file mode 100644
--- /dev/null
+++ b/Engine/Plugins/MeshPainting/Source/MeshPaintingToolset/Public/MeshPaintAdapterFactory.h
@@ -0,0 +1,21 @@
+// Copyright 1998-2019 Epic Games, Inc. All Rights Reserved.
+
+#pragma once
+
+#include "CoreMinimal.h"
+
+class UMeshComponent;
+class IMeshPaintGeometryAdapterFactory;
+class FReferenceCollector;
+
+class MESHPAINTINGTOOLSET_API FMeshPaintAdapterFactory
+{
+public:
+	static TArray<TSharedPtr<IMeshPaintGeometryAdapterFactory>> FactoryList;
+
+public:
+	static TSharedPtr<class IMeshPaintGeometryAdapter> CreateAdapterForMesh(UMeshComponent* InComponent, int32 InPaintingMeshLODIndex);
+	static void InitializeAdapterGlobals();
+	static void AddReferencedObjectsGlobals(FReferenceCollector& Collector);
+	static void CleanupGlobals();
+};
\ No newline at end of file
diff --git a/Engine/Plugins/MeshPainting/Source/MeshPaintingToolset/Public/MeshPaintColorBrush.h b/Engine/Plugins/MeshPainting/Source/MeshPaintingToolset/Public/MeshPaintColorBrush.h
new file mode 100644
--- /dev/null
+++ b/Engine/Plugins/MeshPainting/Source/MeshPaintingToolset/Public/MeshPaintColorBrush.h
@@ -0,0 +1,281 @@
+// Copyright 1998-2019 Epic Games, Inc. All Rights Reserved.
+
+#pragma once
+
+#include "CoreMinimal.h"
+#include "UObject/NoExportTypes.h"
+#include "InteractiveToolBuilder.h"
+#include "BaseTools/BaseBrushTool.h"
+#include "InteractiveTool.h"
+#include "Delegates/DelegateCombinations.h"
+#include "Components/MeshComponent.h"
+#include "MeshPaintHelpers.h"
+#include "MeshPaintingToolsetTypes.h"
+#include "MeshPaintColorBrush.generated.h"
+
+
+struct FToolBuilderState;
+enum class EMeshPaintModeAction : uint8;
+class IMeshPaintGeometryAdapter;
+
+struct FPaintRayResults
+{
+	FMeshPaintParameters Params;
+	FHitResult BestTraceResult;
+};
+
+UENUM()
+enum class EMeshPaintWeightTypes : uint8
+{
+	/** Lerp Between Two Textures using Alpha Value */
+	AlphaLerp = 2 UMETA(DisplayName = "Alpha (Two Textures)"),
+
+	/** Weighting Three Textures according to Channels*/
+	RGB = 3 UMETA(DisplayName = "RGB (Three Textures)"),
+
+	/**  Weighting Four Textures according to Channels*/
+	ARGB = 4 UMETA(DisplayName = "ARGB (Four Textures)"),
+
+	/**  Weighting Five Textures according to Channels */
+	OneMinusARGB = 5 UMETA(DisplayName = "ARGB - 1 (Five Textures)")
+};
+
+UENUM()
+enum class EMeshPaintTextureIndex : uint8
+{
+	TextureOne = 0,
+	TextureTwo,
+	TextureThree,
+	TextureFour,
+	TextureFive
+};
+
+DECLARE_DELEGATE_RetVal(TArray<UMeshComponent*>, FGetSelectedMeshComponents);
+/**
+ *
+ */
+UCLASS()
+class MESHPAINTINGTOOLSET_API UMeshColorPaintingToolBuilder : public UInteractiveToolBuilder
+{
+	GENERATED_BODY()
+
+public:
+	virtual bool CanBuildTool(const FToolBuilderState& SceneState) const override;
+	virtual UInteractiveTool* BuildTool(const FToolBuilderState& SceneState) const override;
+};
+
+UCLASS()
+class MESHPAINTINGTOOLSET_API UMeshWeightPaintingToolBuilder : public UInteractiveToolBuilder
+{
+	GENERATED_BODY()
+
+public:
+	virtual bool CanBuildTool(const FToolBuilderState& SceneState) const override;
+	virtual UInteractiveTool* BuildTool(const FToolBuilderState& SceneState) const override;
+};
+
+
+
+UCLASS()
+class MESHPAINTINGTOOLSET_API UMeshVertexPaintingToolProperties : public UBrushBaseProperties
+{
+	GENERATED_BODY()
+
+public:
+	UMeshVertexPaintingToolProperties();
+
+	/** Color used for Applying Vertex Color Painting */
+	UPROPERTY(EditAnywhere, Category = VertexPainting)
+	FLinearColor PaintColor;
+
+	/** Color used for Erasing Vertex Color Painting */
+	UPROPERTY(EditAnywhere, Category = VertexPainting)
+	FLinearColor EraseColor;
+
+	/** Whether back-facing triangles should be ignored */
+	UPROPERTY(EditAnywhere, Category = Brush, meta = (DisplayName = "Ignore Back-Facing"))
+	bool bOnlyFrontFacingTriangles;
+
+
+	virtual void SaveProperties(UInteractiveTool* SaveFromTool) override;
+	virtual void RestoreProperties(UInteractiveTool* RestoreToTool) override;
+
+};
+
+
+UCLASS()
+class MESHPAINTINGTOOLSET_API UMeshColorPaintingToolProperties : public UMeshVertexPaintingToolProperties
+{
+	GENERATED_BODY()
+
+public:
+	UMeshColorPaintingToolProperties();
+
+
+	/** Whether or not to apply Vertex Color Painting to the Red Channel */
+	UPROPERTY(EditAnywhere, Category = ColorPainting, DisplayName = "Red")
+	bool bWriteRed;
+
+	/** Whether or not to apply Vertex Color Painting to the Green Channel */
+	UPROPERTY(EditAnywhere, Category = ColorPainting, DisplayName = "Green")
+	bool bWriteGreen;
+
+	/** Whether or not to apply Vertex Color Painting to the Blue Channel */
+	UPROPERTY(EditAnywhere, Category = ColorPainting, DisplayName = "Blue")
+	bool bWriteBlue;
+
+	/** Whether or not to apply Vertex Color Painting to the Alpha Channel */
+	UPROPERTY(EditAnywhere, Category = ColorPainting, DisplayName = "Alpha")
+	bool bWriteAlpha;
+
+	/** When unchecked the painting on the base LOD will be propagate automatically to all other LODs when exiting the mode or changing the selection */
+	UPROPERTY(EditAnywhere, Category = Painting)
+	bool bPaintOnSpecificLOD;
+
+	/** LOD Index to which should specifically be painted */
+	UPROPERTY(EditAnywhere, Category = Painting, meta = (UIMin = "0", ClampMin = "0", EditCondition = "bPaintOnSpecificLOD"))
+	int32 LODIndex;
+
+	virtual void SaveProperties(UInteractiveTool* SaveFromTool) override;
+	virtual void RestoreProperties(UInteractiveTool* RestoreToTool) override;
+};
+
+UCLASS()
+class MESHPAINTINGTOOLSET_API UMeshWeightPaintingToolProperties : public UMeshVertexPaintingToolProperties
+{
+	GENERATED_BODY()
+
+public:
+	UMeshWeightPaintingToolProperties();
+
+	/** Texture Blend Weight Painting Mode */
+	UPROPERTY(EditAnywhere, Category = WeightPainting, meta = (EnumCondition = 1))
+	EMeshPaintWeightTypes TextureWeightType;
+
+	/** Texture Blend Weight index which should be applied during Painting */
+	UPROPERTY(EditAnywhere, Category = WeightPainting, meta = (EnumCondition = 1))
+	EMeshPaintTextureIndex PaintTextureWeightIndex;
+
+	/** Texture Blend Weight index which should be erased during Painting */
+	UPROPERTY(EditAnywhere, Category = WeightPainting, meta = (EnumCondition = 1))
+	EMeshPaintTextureIndex EraseTextureWeightIndex;
+
+	virtual void SaveProperties(UInteractiveTool* SaveFromTool) override;
+	virtual void RestoreProperties(UInteractiveTool* RestoreToTool) override;
+};
+
+UCLASS(Abstract)
+class MESHPAINTINGTOOLSET_API UMeshVertexPaintingTool : public UBaseBrushTool
+{
+	GENERATED_BODY()
+
+public:
+	UMeshVertexPaintingTool();
+
+	virtual void Setup() override;
+	virtual void Shutdown(EToolShutdownType ShutdownType) override;
+	virtual void Render(IToolsContextRenderAPI* RenderAPI) override;
+	virtual void Tick(float DeltaTime) override;
+	virtual bool HasCancel() const override { return false; }
+	virtual bool HasAccept() const override;
+	virtual bool CanAccept() const override;
+	virtual void OnBeginDrag(const FRay& Ray) override;
+	virtual void OnUpdateDrag(const FRay& Ray) override;
+	virtual void OnEndDrag(const FRay& Ray) override;
+	virtual	bool HitTest(const FRay& Ray, FHitResult& OutHit) override;
+	virtual void OnPropertyModified(UObject* PropertySet, UProperty* Property) override;
+	virtual bool IsPainting() const
+	{
+		return bArePainting;
+	}
+
+
+	FSimpleDelegate OnPaintingFinished()
+	{
+		return OnPaintingFinishedDelegate;
+	}
+
+protected:
+	virtual void SetAdditionalPaintParameters(FMeshPaintParameters& InPaintParameters) {};
+	virtual void FinishPainting();
+	void UpdateResult();
+	double CalculateTargetEdgeLength(int TargetTriCount);
+	bool Paint(const FVector& InRayOrigin, const FVector& InRayDirection);
+	bool Paint(const TArrayView<TPair<FVector, FVector>>& Rays);
+	virtual void CacheSelectionData() {};
+	/** Per vertex action function used for painting vertex data */
+	void ApplyVertexData(FPerVertexPaintActionArgs& InArgs, int32 VertexIndex, FMeshPaintParameters Parameters);
+
+
+protected:
+	double InitialMeshArea;
+	bool bResultValid;
+	bool bStampPending;
+	bool bInDrag;
+	FRay PendingStampRay;
+
+private:
+	bool PaintInternal(const TArrayView<TPair<FVector, FVector>>& Rays, EMeshPaintModeAction PaintAction, float PaintStrength);
+	
+private:
+	UPROPERTY(Transient)
+	UMeshVertexPaintingToolProperties* VertexProperties;
+	/** Flag for whether or not we are currently painting */
+	bool bArePainting;
+	bool bDoRestoreRenTargets;
+	/** Time kept since the user has started painting */
+	float TimeSinceStartedPainting;
+	/** Overall time value kept for drawing effects */
+	float Time;
+	FHitResult LastBestHitResult;
+	FSimpleDelegate OnPaintingFinishedDelegate;
+};
+
+UCLASS()
+class MESHPAINTINGTOOLSET_API UMeshColorPaintingTool : public UMeshVertexPaintingTool
+{
+	GENERATED_BODY()
+
+public:
+	UMeshColorPaintingTool();
+	virtual void Setup() override;
+	void LODPaintStateChanged(const bool bLODPaintingEnabled);
+	int32 GetMaxLODIndexToPaint() const;
+	void PaintLODChanged();
+	int32 GetCachedLODIndex() const
+	{
+		return CachedLODIndex;
+	}
+
+protected:
+	virtual void CacheSelectionData() override;
+	virtual void SetAdditionalPaintParameters(FMeshPaintParameters& InPaintParameters) override;
+	void ApplyForcedLODIndex(int32 ForcedLODIndex);
+
+private:
+	UPROPERTY(Transient)
+	UMeshColorPaintingToolProperties* ColorProperties;
+	/** Current LOD index used for painting / forcing */
+	int32 CachedLODIndex;
+	/** Whether or not a specific LOD level should be forced */
+	bool bCachedForceLOD;
+};
+
+
+UCLASS()
+class MESHPAINTINGTOOLSET_API UMeshWeightPaintingTool : public UMeshVertexPaintingTool
+{
+	GENERATED_BODY()
+
+public:
+	UMeshWeightPaintingTool();
+	virtual void Setup() override;
+
+protected:
+	virtual void CacheSelectionData() override;
+	virtual void SetAdditionalPaintParameters(FMeshPaintParameters& InPaintParameters);
+
+private:
+	UPROPERTY(Transient)
+	UMeshWeightPaintingToolProperties* WeightProperties;
+};
\ No newline at end of file
diff --git a/Engine/Plugins/MeshPainting/Source/MeshPaintingToolset/Public/MeshPaintHelpers.h b/Engine/Plugins/MeshPainting/Source/MeshPaintingToolset/Public/MeshPaintHelpers.h
new file mode 100644
--- /dev/null
+++ b/Engine/Plugins/MeshPainting/Source/MeshPaintingToolset/Public/MeshPaintHelpers.h
@@ -0,0 +1,307 @@
+// Copyright 1998-2019 Epic Games, Inc. All Rights Reserved.
+
+#pragma once
+
+#include "CoreMinimal.h"
+#include "Engine/EngineTypes.h"
+#include "MeshPaintingToolsetTypes.h"
+#include "Kismet/BlueprintFunctionLibrary.h"
+#include "IMeshPaintGeometryAdapter.h"
+#include "InteractiveToolManager.h"
+#include "InteractiveToolsContext.h"
+#include "Engine/StaticMesh.h"
+#include "MeshPaintHelpers.generated.h"
+
+class FMeshPaintParameters;
+class UMeshVertexColorImportOptions;
+class UTexture2D;
+class UStaticMeshComponent;
+class USkeletalMesh;
+class IMeshPaintGeometryAdapter;
+class UPaintBrushSettings;
+class FEditorViewportClient;
+class UMeshComponent;
+class USkeletalMeshComponent;
+class UViewportInteractor;
+class FViewport;
+class FPrimitiveDrawInterface;
+class FSceneView;
+struct FStaticMeshComponentLODInfo;
+class UMeshVertexPaintingToolProperties;
+
+enum class EMeshPaintDataColorViewMode : uint8;
+
+/** struct used to store the color data copied from mesh instance to mesh instance */
+struct FPerLODVertexColorData
+{
+	TArray< FColor > ColorsByIndex;
+	TMap<FVector, FColor> ColorsByPosition;
+};
+
+/** struct used to store the color data copied from mesh component to mesh component */
+struct FPerComponentVertexColorData
+{
+	FPerComponentVertexColorData(const UStaticMesh* InStaticMesh, int32 InComponentIndex)
+		: OriginalMesh(InStaticMesh)
+		, ComponentIndex(InComponentIndex)
+	{
+	}
+
+	/** We match up components by the mesh they use */
+	TWeakObjectPtr<const UStaticMesh> OriginalMesh;
+
+	/** We also match by component index */
+	int32 ComponentIndex;
+
+	/** Vertex colors by LOD */
+	TArray<FPerLODVertexColorData> PerLODVertexColorData;
+};
+
+/** Struct to hold MeshPaint settings on a per mesh basis */
+struct FInstanceTexturePaintSettings
+{
+	UTexture2D* SelectedTexture;
+	int32 SelectedUVChannel;
+
+	FInstanceTexturePaintSettings()
+		: SelectedTexture(nullptr)
+		, SelectedUVChannel(0)
+	{}
+	FInstanceTexturePaintSettings(UTexture2D* InSelectedTexture, int32 InSelectedUVSet)
+		: SelectedTexture(InSelectedTexture)
+		, SelectedUVChannel(InSelectedUVSet)
+	{}
+
+	void operator=(const FInstanceTexturePaintSettings& SrcSettings)
+	{
+		SelectedTexture = SrcSettings.SelectedTexture;
+		SelectedUVChannel = SrcSettings.SelectedUVChannel;
+	}
+};
+
+
+UENUM()
+enum class ETexturePaintWeightTypes : uint8
+{
+	/** Lerp Between Two Textures using Alpha Value */
+	AlphaLerp = 2 UMETA(DisplayName = "Alpha (Two Textures)"),
+
+	/** Weighting Three Textures according to Channels*/
+	RGB = 3 UMETA(DisplayName = "RGB (Three Textures)"),
+
+	/**  Weighting Four Textures according to Channels*/
+	ARGB = 4 UMETA(DisplayName = "ARGB (Four Textures)"),
+
+	/**  Weighting Five Textures according to Channels */
+	OneMinusARGB = 5 UMETA(DisplayName = "ARGB - 1 (Five Textures)")
+};
+
+UENUM()
+enum class ETexturePaintWeightIndex : uint8
+{
+	TextureOne = 0,
+	TextureTwo,
+	TextureThree,
+	TextureFour,
+	TextureFive
+};
+
+/** Parameters for paint actions, stored together for convenience */
+struct FPerVertexPaintActionArgs
+{
+	IMeshPaintGeometryAdapter* Adapter;
+	UMeshVertexPaintingToolProperties* BrushProperties;
+	FVector CameraPosition;
+	FHitResult HitResult;
+	EMeshPaintModeAction Action;
+};
+
+/** Delegates used to call per-vertex/triangle actions */
+DECLARE_DELEGATE_TwoParams(FPerVertexPaintAction, FPerVertexPaintActionArgs& /*Args*/, int32 /*VertexIndex*/);
+DECLARE_DELEGATE_ThreeParams(FPerTrianglePaintAction, IMeshPaintGeometryAdapter* /*Adapter*/, int32 /*TriangleIndex*/, const int32[3] /*Vertex Indices*/);
+
+class MESHPAINTINGTOOLSET_API UMeshPaintingToolset : public UBlueprintFunctionLibrary
+{
+public:
+	/** Removes vertex colors associated with the object */
+	static void RemoveInstanceVertexColors(UObject* Obj);
+
+	/** Removes vertex colors associated with the static mesh component */
+	static void RemoveComponentInstanceVertexColors(UStaticMeshComponent* StaticMeshComponent);
+	
+	/** Propagates per-instance vertex colors to the underlying Static Mesh for the given LOD Index */
+	static bool PropagateColorsToRawMesh(UStaticMesh* StaticMesh, int32 LODIndex, FStaticMeshComponentLODInfo& ComponentLODInfo);	
+
+	/** Retrieves the Vertex Color buffer size for the given LOD level in the Static Mesh */
+	static uint32 GetVertexColorBufferSize(UMeshComponent* MeshComponent, int32 LODIndex, bool bInstance);
+
+	/** Retrieves the vertex positions from the given LOD level in the Static Mesh */
+	static TArray<FVector> GetVerticesForLOD(const UStaticMesh* StaticMesh, int32 LODIndex);
+
+	/** Retrieves the vertex colors from the given LOD level in the Static Mesh */
+	static TArray<FColor> GetColorDataForLOD(const UStaticMesh* StaticMesh, int32 LODIndex);
+
+	/** Retrieves the per-instance vertex colors from the given LOD level in the StaticMeshComponent */
+	static TArray<FColor> GetInstanceColorDataForLOD(const UStaticMeshComponent* MeshComponent, int32 LODIndex);
+
+	/** Sets the specific (LOD Index) per-instance vertex colors for the given StaticMeshComponent to the supplied Color array */
+	static void SetInstanceColorDataForLOD(UStaticMeshComponent* MeshComponent, int32 LODIndex, const TArray<FColor>& Colors);	
+	
+	/** Sets the specific (LOD Index) per-instance vertex colors for the given StaticMeshComponent to a single Color value */
+	static void SetInstanceColorDataForLOD(UStaticMeshComponent* MeshComponent, int32 LODIndex, const FColor FillColor, const FColor MaskColor);
+	
+	/** Fills all vertex colors for all LODs found in the given mesh component with Fill Color */
+	static void FillStaticMeshVertexColors(UStaticMeshComponent* MeshComponent, int32 LODIndex, const FColor FillColor, const FColor MaskColor);
+	static void FillSkeletalMeshVertexColors(USkeletalMeshComponent* MeshComponent, int32 LODIndex, const FColor FillColor, const FColor MaskColor);
+	
+	/** Sets all vertex colors for a specific LOD level in the SkeletalMesh to FillColor */
+	static void SetColorDataForLOD(USkeletalMesh* SkeletalMesh, int32 LODIndex, const FColor FillColor, const FColor MaskColor);
+
+	static void ApplyFillWithMask(FColor& InOutColor, const FColor& MaskColor, const FColor& FillColor);
+
+		
+	/** Checks whether or not the given Viewport Client is a VR editor viewport client */
+	static bool IsInVRMode(const FEditorViewportClient* ViewportClient);
+
+	/** Forces the component to render LOD level at LODIndex instead of the view-based LOD level ( X = 0 means do not force the LOD, X > 0 means force the lod to X - 1 ) */
+	static void ForceRenderMeshLOD(UMeshComponent* Component, int32 LODIndex);
+
+	/** Clears all texture overrides for this component. */
+	static void ClearMeshTextureOverrides(const IMeshPaintGeometryAdapter& GeometryInfo, UMeshComponent* InMeshComponent);
+
+	/** Applies vertex color painting found on LOD 0 to all lower LODs. */
+	static void ApplyVertexColorsToAllLODs(IMeshPaintGeometryAdapter& GeometryInfo, UMeshComponent* InMeshComponent);
+
+	/** Applies the vertex colors found in LOD level 0 to all contained LOD levels in the StaticMeshComponent */
+	static void ApplyVertexColorsToAllLODs(IMeshPaintGeometryAdapter& GeometryInfo, UStaticMeshComponent* StaticMeshComponent);
+
+	/** Applies the vertex colors found in LOD level 0 to all contained LOD levels in the SkeletalMeshComponent */
+	static void ApplyVertexColorsToAllLODs(IMeshPaintGeometryAdapter& GeometryInfo, USkeletalMeshComponent* SkeletalMeshComponent);
+
+	/** Returns the number of Mesh LODs for the given MeshComponent */
+	static int32 GetNumberOfLODs(const UMeshComponent* MeshComponent);
+
+	/** OutNumLODs is set to number of Mesh LODs for the given MeshComponent and returns true, or returns false of given mesh component has no valid LODs */
+	static bool TryGetNumberOfLODs(const UMeshComponent* MeshComponent, int32& OutNumLODs);
+	
+	/** Returns the number of Texture Coordinates for the given MeshComponent */
+	static int32 GetNumberOfUVs(const UMeshComponent* MeshComponent, int32 LODIndex);
+
+	/** Checks whether or not the mesh components contains per lod colors (for all LODs)*/
+	static bool DoesMeshComponentContainPerLODColors(const UMeshComponent* MeshComponent);
+
+	/** Retrieves the number of bytes used to store the per-instance LOD vertex color data from the static mesh component */
+	static void GetInstanceColorDataInfo(const UStaticMeshComponent* StaticMeshComponent, int32 LODIndex, int32& OutTotalInstanceVertexColorBytes);
+
+	/** Given arguments for an action, and an action - retrieves influences vertices and applies Action to them */
+	static bool ApplyPerVertexPaintAction(FPerVertexPaintActionArgs& InArgs, FPerVertexPaintAction Action);
+
+	static bool GetPerVertexPaintInfluencedVertices(FPerVertexPaintActionArgs& InArgs, TSet<int32>& InfluencedVertices);
+
+	/** Given the adapter, settings and view-information retrieves influences triangles and applies Action to them */
+	static bool ApplyPerTrianglePaintAction(IMeshPaintGeometryAdapter* Adapter, const FVector& CameraPosition, const FVector& HitPosition, const UMeshVertexPaintingToolProperties* Settings, FPerTrianglePaintAction Action);
+
+	/** Applies vertex painting to InOutvertexColor according to the given parameters  */
+	static bool PaintVertex(const FVector& InVertexPosition, const FMeshPaintParameters& InParams, FColor& InOutVertexColor);
+
+	/** Applies Vertex Color Painting according to the given parameters */
+	static void ApplyVertexColorPaint(const FMeshPaintParameters &InParams, const FLinearColor &OldColor, FLinearColor &NewColor, const float PaintAmount);
+
+	/** Applies Vertex Blend Weight Painting according to the given parameters */
+	static void ApplyVertexWeightPaint(const FMeshPaintParameters &InParams, const FLinearColor &OldColor, FLinearColor &NewColor, const float PaintAmount);
+
+	/** Generate texture weight color for given number of weights and the to-paint index */
+	static FLinearColor GenerateColorForTextureWeight(const int32 NumWeights, const int32 WeightIndex);
+
+	/** Computes the Paint power multiplier value */
+	static float ComputePaintMultiplier(float SquaredDistanceToVertex2D, float BrushStrength, float BrushInnerRadius, float BrushRadialFalloff, float BrushInnerDepth, float BrushDepthFallof, float VertexDepthToBrush);
+
+	/** Checks whether or not a point is influenced by the painting brush according to the given parameters*/
+	static bool IsPointInfluencedByBrush(const FVector& InPosition, const FMeshPaintParameters& InParams, float& OutSquaredDistanceToVertex2D, float& OutVertexDepthToBrush);
+
+	static bool IsPointInfluencedByBrush(const FVector2D& BrushSpacePosition, const float BrushRadiusSquared, float& OutInRangeValue);
+
+	template<typename T>
+	static void ApplyBrushToVertex(const FVector& VertexPosition, const FMatrix& InverseBrushMatrix, const float BrushRadius, const float BrushFalloffAmount, const float BrushStrength, const T& PaintValue, T& InOutValue);
+
+public:
+
+
+
+	/** Helper function to retrieve vertex color from a UTexture given a UVCoordinate */
+	static FColor PickVertexColorFromTextureData(const uint8* MipData, const FVector2D& UVCoordinate, const UTexture2D* Texture, const FColor ColorMask);	
+};
+
+template<typename T>
+void UMeshPaintingToolset::ApplyBrushToVertex(const FVector& VertexPosition, const FMatrix& InverseBrushMatrix, const float BrushRadius, const float BrushFalloffAmount, const float BrushStrength, const T& PaintValue, T& InOutValue)
+{
+	const FVector BrushSpacePosition = InverseBrushMatrix.TransformPosition(VertexPosition);
+	const FVector2D BrushSpacePosition2D(BrushSpacePosition.X, BrushSpacePosition.Y);
+		
+	float InfluencedValue = 0.0f;
+	if (IsPointInfluencedByBrush(BrushSpacePosition2D, BrushRadius * BrushRadius, InfluencedValue))
+	{
+		float InnerBrushRadius = BrushFalloffAmount * BrushRadius;
+		float PaintStrength = MeshPaintHelpers::ComputePaintMultiplier(BrushSpacePosition2D.SizeSquared(), BrushStrength, InnerBrushRadius, BrushRadius - InnerBrushRadius, 1.0f, 1.0f, 1.0f);
+
+		const T OldValue = InOutValue;
+		InOutValue = FMath::LerpStable(OldValue, PaintValue, PaintStrength);
+	}	
+}
+
+UCLASS()
+class MESHPAINTINGTOOLSET_API UMeshToolsContext : public UInteractiveToolsContext
+{
+	GENERATED_BODY()
+
+public:
+	UMeshToolsContext();
+};
+
+UCLASS()
+class MESHPAINTINGTOOLSET_API UMeshToolManager : public UInteractiveToolManager
+{
+	GENERATED_BODY()
+
+public:
+	UMeshToolManager();
+	virtual void Shutdown() override;
+
+	/** Map of geometry adapters for each selected mesh component */
+	TMap<UMeshComponent*, TSharedPtr<IMeshPaintGeometryAdapter>> GetComponentToAdapterMap() const;
+	TSharedPtr<IMeshPaintGeometryAdapter> GetAdapterForComponent(UMeshComponent* InComponent);
+	void AddToComponentToAdapterMap(UMeshComponent* InComponent, TSharedPtr<IMeshPaintGeometryAdapter> InAdapter);
+
+	TArray<UMeshComponent*> GetSelectedMeshComponents() const;
+	void AddSelectedMeshComponents(const TArray<UMeshComponent*>& InComponents);
+	bool FindHitResult(const FRay Ray, FHitResult& BestTraceResult);
+	void ClearSelectedMeshComponents();
+	TArray<UMeshComponent*> GetPaintableMeshComponents() const;
+	void AddPaintableMeshComponent(UMeshComponent* InComponent);
+	void ClearPaintableMeshComponents();
+	bool SelectionContainsValidAdapters() const;
+	TArray<FPerComponentVertexColorData> GetCopiedColorsByComponent() const;
+	void SetCopiedColorsByComponent(TArray<FPerComponentVertexColorData>& InCopiedColors);
+	void CacheSelectionData(const int32 PaintLODIndex, const int32 UVChannel);
+	void ResetState();
+	void Refresh();
+	bool SelectionContainsPerLODColors() const { return bSelectionContainsPerLODColors; }
+	void ClearSelectionLODColors() { bSelectionContainsPerLODColors = false; }
+
+public:
+	bool bNeedsRecache;
+
+
+private:
+	void CleanUp();
+
+private:
+	/** Map of geometry adapters for each selected mesh component */
+	TMap<UMeshComponent*, TSharedPtr<IMeshPaintGeometryAdapter>> ComponentToAdapterMap;
+	TArray<UMeshComponent*> SelectedMeshComponents;
+	/** Mesh components within the current selection which are eligible for painting */
+	TArray<UMeshComponent*> PaintableComponents;
+	/** Contains copied vertex color data */
+	TArray<FPerComponentVertexColorData> CopiedColorsByComponent;
+	bool bSelectionContainsPerLODColors;
+};
\ No newline at end of file
diff --git a/Engine/Plugins/MeshPainting/Source/MeshPaintingToolset/Public/MeshPaintSplineMeshAdapter.h b/Engine/Plugins/MeshPainting/Source/MeshPaintingToolset/Public/MeshPaintSplineMeshAdapter.h
new file mode 100644
--- /dev/null
+++ b/Engine/Plugins/MeshPainting/Source/MeshPaintingToolset/Public/MeshPaintSplineMeshAdapter.h
@@ -0,0 +1,24 @@
+// Copyright 1998-2019 Epic Games, Inc. All Rights Reserved.
+
+#pragma once
+
+#include "CoreMinimal.h"
+#include "MeshPaintStaticMeshAdapter.h"
+
+//////////////////////////////////////////////////////////////////////////
+// FMeshPaintGeometryAdapterForSplineMeshes
+
+class MESHPAINTINGTOOLSET_API FMeshPaintGeometryAdapterForSplineMeshes : public FMeshPaintGeometryAdapterForStaticMeshes
+{
+public:
+	virtual bool InitializeVertexData() override;
+};
+
+//////////////////////////////////////////////////////////////////////////
+// FMeshPaintGeometryAdapterForSplineMeshesFactory
+
+class MESHPAINTINGTOOLSET_API FMeshPaintGeometryAdapterForSplineMeshesFactory : public FMeshPaintGeometryAdapterForStaticMeshesFactory
+{
+public:
+	virtual TSharedPtr<IMeshPaintGeometryAdapter> Construct(class UMeshComponent* InComponent, int32 MeshLODIndex) const override;
+};
diff --git a/Engine/Plugins/MeshPainting/Source/MeshPaintingToolset/Public/MeshPaintingToolsetTypes.h b/Engine/Plugins/MeshPainting/Source/MeshPaintingToolset/Public/MeshPaintingToolsetTypes.h
new file mode 100644
--- /dev/null
+++ b/Engine/Plugins/MeshPainting/Source/MeshPaintingToolset/Public/MeshPaintingToolsetTypes.h
@@ -0,0 +1,198 @@
+// Copyright 1998-2019 Epic Games, Inc. All Rights Reserved.
+
+#pragma once
+
+#include "Engine/Texture.h"
+#include "Engine/TextureRenderTarget2D.h"
+#include "Materials/MaterialInterface.h"
+
+#include "MeshPaintingToolsetTypes.generated.h"
+
+class FMeshPaintParameters;
+DECLARE_MULTICAST_DELEGATE_FourParams(FApplyVertexPaintData, const FMeshPaintParameters& /* &InParams*/, const FLinearColor& /*OldColor*/, FLinearColor& /*NewColor*/, const float /*PaintAmount*/);
+
+/** Mesh paint color view modes (somewhat maps to EVertexColorViewMode engine enum.) */
+UENUM()
+enum class EMeshPaintDataColorViewMode : uint8
+{
+	/** Normal view mode (vertex color visualization off) */
+	Normal UMETA(DisplayName = "Off"),
+
+	/** RGB only */
+	RGB UMETA(DisplayName = "RGB Channels"),
+
+	/** Alpha only */
+	Alpha UMETA(DisplayName = "Alpha Channel"),
+
+	/** Red only */
+	Red UMETA(DisplayName = "Red Channel"),
+
+	/** Green only */
+	Green UMETA(DisplayName = "Green Channel"),
+
+	/** Blue only */
+	Blue UMETA(DisplayName = "Blue Channel"),
+};
+
+/** Mesh painting action (paint, erase) */
+enum class EMeshPaintModeAction : uint8
+{
+	/** Paint (add color or increase blending weight) */
+	Paint,
+
+	/** Erase (remove color or decrease blending weight) */
+	Erase
+};
+
+/** Vertex paint target */
+UENUM()
+enum class EMeshVertexPaintModeTarget : uint8
+{
+	/** Paint the static mesh component instance in the level */
+	ComponentInstance,
+
+	/** Paint the actual static mesh asset */
+	Mesh
+};
+
+/** Mesh paint parameters */
+class FMeshPaintParameters
+{
+public:
+	EMeshPaintModeAction PaintAction;
+	FVector BrushPosition;
+	FVector BrushNormal;
+	FLinearColor BrushColor;
+	float SquaredBrushRadius;
+	float BrushRadialFalloffRange;
+	float InnerBrushRadius;
+	float BrushDepth;
+	float BrushDepthFalloffRange;
+	float InnerBrushDepth;
+	float BrushStrength;
+	FMatrix BrushToWorldMatrix;
+	FMatrix InverseBrushToWorldMatrix;
+	bool bWriteRed;
+	bool bWriteGreen;
+	bool bWriteBlue;
+	bool bWriteAlpha;
+	int32 TotalWeightCount;
+	int32 PaintWeightIndex;
+	int32 UVChannel;
+	FApplyVertexPaintData ApplyVertexDataDelegate;
+};
+
+/** Structure used to hold per-triangle data for texture painting */
+struct FTexturePaintTriangleInfo
+{
+	FVector TriVertices[3];
+	FVector2D TrianglePoints[3];
+	FVector2D TriUVs[3];
+};
+
+/** Structure used to house and compare Texture and UV channel pairs */
+struct FPaintableTexture
+{
+	UTexture*	Texture;
+	int32		UVChannelIndex;
+
+	FPaintableTexture(UTexture* InTexture = nullptr, uint32 InUVChannelIndex = 0)
+		: Texture(InTexture)
+		, UVChannelIndex(InUVChannelIndex)
+	{}
+
+	/** Overloaded equality operator for use with TArrays Contains method. */
+	bool operator==(const FPaintableTexture& rhs) const
+	{
+		return (Texture == rhs.Texture);
+		/* && (UVChannelIndex == rhs.UVChannelIndex);*/// if we compare UVChannel we would have to duplicate the texture
+	}
+};
+
+struct FPaintTexture2DData
+{
+	/** The original texture that we're painting */
+	UTexture2D* PaintingTexture2D;
+	bool bIsPaintingTexture2DModified;
+
+	/** A copy of the original texture we're painting, used for restoration. */
+	UTexture2D* PaintingTexture2DDuplicate;
+
+	/** Render target texture for painting */
+	UTextureRenderTarget2D* PaintRenderTargetTexture;
+
+	/** Render target texture used as an input while painting that contains a clone of the original image */
+	UTextureRenderTarget2D* CloneRenderTargetTexture;
+
+	/** List of materials we are painting on */
+	TArray< UMaterialInterface* > PaintingMaterials;
+
+	/** Default ctor */
+	FPaintTexture2DData() :
+		PaintingTexture2D(NULL),
+		bIsPaintingTexture2DModified(false),
+		PaintingTexture2DDuplicate(nullptr),
+		PaintRenderTargetTexture(nullptr),
+		CloneRenderTargetTexture(nullptr)
+	{}
+
+	FPaintTexture2DData(UTexture2D* InPaintingTexture2D, bool InbIsPaintingTexture2DModified = false) :
+		PaintingTexture2D(InPaintingTexture2D),
+		bIsPaintingTexture2DModified(InbIsPaintingTexture2DModified),
+		PaintRenderTargetTexture(nullptr),
+		CloneRenderTargetTexture(nullptr)
+	{}
+
+	/** Serializer */
+	void AddReferencedObjects(FReferenceCollector& Collector)
+	{
+		// @todo MeshPaint: We're relying on GC to clean up render targets, can we free up remote memory more quickly?
+		Collector.AddReferencedObject(PaintingTexture2D);
+		Collector.AddReferencedObject(PaintRenderTargetTexture);
+		Collector.AddReferencedObject(CloneRenderTargetTexture);
+		for (int32 Index = 0; Index < PaintingMaterials.Num(); Index++)
+		{
+			Collector.AddReferencedObject(PaintingMaterials[Index]);
+		}
+	}
+};
+
+namespace MeshPaintDefs
+{
+	// Design constraints
+
+	// Currently we never support more than five channels (R, G, B, A, OneMinusTotal)
+	static const int32 MaxSupportedPhysicalWeights = 4;
+	static const int32 MaxSupportedWeights = MaxSupportedPhysicalWeights + 1;
+}
+
+/**
+*  Wrapper to expose texture targets to WPF code.
+*/
+struct FTextureTargetListInfo
+{
+	UTexture2D* TextureData;
+	bool bIsSelected;
+	uint32 UndoCount;
+	uint32 UVChannelIndex;
+	FTextureTargetListInfo(UTexture2D* InTextureData, int32 InUVChannelIndex, bool InbIsSelected = false)
+		: TextureData(InTextureData)
+		, bIsSelected(InbIsSelected)
+		, UndoCount(0)
+		, UVChannelIndex(InUVChannelIndex)
+	{}
+};
+
+/**
+*  Wrapper to store which of a meshes materials is selected as well as the total number of materials.
+*/
+struct FMeshSelectedMaterialInfo
+{
+	int32 NumMaterials;
+	int32 SelectedMaterialIndex;
+
+	FMeshSelectedMaterialInfo(int32 InNumMaterials)
+		: NumMaterials(InNumMaterials)
+		, SelectedMaterialIndex(0)
+	{}
+};
\ No newline at end of file
diff --git a/Engine/Plugins/MeshPainting/Source/MeshPaintingToolset/Public/MeshSelect.h b/Engine/Plugins/MeshPainting/Source/MeshPaintingToolset/Public/MeshSelect.h
new file mode 100644
--- /dev/null
+++ b/Engine/Plugins/MeshPainting/Source/MeshPaintingToolset/Public/MeshSelect.h
@@ -0,0 +1,57 @@
+// Copyright 1998-2019 Epic Games, Inc. All Rights Reserved.
+
+#pragma once
+
+#include "CoreMinimal.h"
+#include "BaseTools/SingleClickTool.h"
+#include "MeshSelect.generated.h"
+
+
+
+/**
+ * Builder for UMeshClickToolBuilder
+ */
+UCLASS()
+class MESHPAINTINGTOOLSET_API UMeshClickToolBuilder : public USingleClickToolBuilder
+{
+	GENERATED_BODY()
+
+public:
+	virtual bool CanBuildTool(const FToolBuilderState& SceneState) const override;
+	virtual UInteractiveTool* BuildTool(const FToolBuilderState& SceneState) const override;
+};
+
+
+/**
+ * USingleClickTool is perhaps the simplest possible interactive tool. It simply
+ * reacts to default primary button clicks for the active device (eg left-mouse clicks).
+ *
+ * The function ::IsHitByClick() determines what is clickable by this Tool. The default is
+ * to return true, which means the click will activate anywhere (the Tool itself has no
+ * notion of Actors, Components, etc). You can override this function to, for example,
+ * filter out clicks that don't hit a target object, etc.
+ *
+ * The function ::OnClicked() implements the action that will occur when a click happens.
+ * You must override this to implement any kind of useful behavior.
+ */
+UCLASS()
+class MESHPAINTINGTOOLSET_API UMeshClickTool : public USingleClickTool
+{
+	GENERATED_BODY()
+
+public:
+	UMeshClickTool();
+
+	// USingleClickTool overrides
+	virtual	void Setup() override;
+	virtual void OnUpdateModifierState(int ModifierID, bool bIsOn) override;
+	virtual bool IsHitByClick(const FInputDeviceRay& ClickPos) override;
+	virtual void OnClicked(const FInputDeviceRay& ClickPos) override;
+
+protected:
+	// flags used to identify modifier keys/buttons
+	static const int AdditiveSelectionModifier = 1;
+	bool bAddToSelectionSet;
+
+};
+
diff --git a/Engine/Plugins/MeshPainting/Source/MeshPaintingToolset/Public/TexturePaintHelpers.h b/Engine/Plugins/MeshPainting/Source/MeshPaintingToolset/Public/TexturePaintHelpers.h
new file mode 100644
--- /dev/null
+++ b/Engine/Plugins/MeshPainting/Source/MeshPaintingToolset/Public/TexturePaintHelpers.h
@@ -0,0 +1,88 @@
+// Copyright 1998-2019 Epic Games, Inc. All Rights Reserved.
+
+#pragma once
+
+#include "BatchedElements.h"
+#include "MeshPaintRendering.h"
+
+class UMeshComponent;
+class UTexture;
+class UTexture2D;
+class UTextureRenderTarget2D;
+
+/** Batched element parameters for texture paint shaders used for paint blending and paint mask generation */
+class MESHPAINTINGTOOLSET_API FMeshPaintBatchedElementParameters : public FBatchedElementParameters
+{
+public:
+	/** Binds vertex and pixel shaders for this element */
+	virtual void BindShaders(FRHICommandList& RHICmdList, FGraphicsPipelineStateInitializer& GraphicsPSOInit, ERHIFeatureLevel::Type InFeatureLevel, const FMatrix& InTransform, const float InGamma, const FMatrix& ColorWeights, const FTexture* Texture) override
+	{
+		MeshPaintRendering::SetMeshPaintShaders(RHICmdList, GraphicsPSOInit, InFeatureLevel, InTransform, InGamma, ShaderParams);
+	}
+
+public:
+
+	/** Shader parameters */
+	MeshPaintRendering::FMeshPaintShaderParameters ShaderParams;
+};
+
+/** Batched element parameters for texture paint shaders used for texture dilation */
+class MESHPAINTINGTOOLSET_API FMeshPaintDilateBatchedElementParameters : public FBatchedElementParameters
+{
+public:
+	/** Binds vertex and pixel shaders for this element */
+	virtual void BindShaders(FRHICommandList& RHICmdList, FGraphicsPipelineStateInitializer& GraphicsPSOInit, ERHIFeatureLevel::Type InFeatureLevel, const FMatrix& InTransform, const float InGamma, const FMatrix& ColorWeights, const FTexture* Texture) override
+	{
+		MeshPaintRendering::SetMeshPaintDilateShaders(RHICmdList, GraphicsPSOInit, InFeatureLevel, InTransform, InGamma, ShaderParams);
+	}
+
+public:
+
+	/** Shader parameters */
+	MeshPaintRendering::FMeshPaintDilateShaderParameters ShaderParams;
+};
+
+/** Helper struct to store mesh section information in*/
+struct FTexturePaintMeshSectionInfo
+{
+	/** First vertex index in section */
+	int32 FirstIndex;
+
+	/** Last vertex index in section */
+	int32 LastIndex;
+};
+
+class IMeshPaintGeometryAdapter;
+struct FPaintableTexture;
+
+/** Helpers functions for texture painting functionality */
+class MESHPAINTINGTOOLSET_API TexturePaintHelpers
+{
+public:
+	/** Static: Copies a texture to a render target texture */
+	static void CopyTextureToRenderTargetTexture(UTexture* SourceTexture, UTextureRenderTarget2D* RenderTargetTexture, ERHIFeatureLevel::Type FeatureLevel);
+
+	/** Will generate a mask texture, used for texture dilation, and store it in the passed in render target */
+	static bool GenerateSeamMask(UMeshComponent* MeshComponent, int32 UVSet, UTextureRenderTarget2D* SeamRenderTexture, UTexture2D* Texture, UTextureRenderTarget2D* RenderTargetTexture );
+
+	/** Static: Creates a temporary texture used to transfer data to a render target in memory */
+	static UTexture2D* CreateTempUncompressedTexture(UTexture2D* SourceTexture);
+
+	/** Makes sure that the render target is ready to paint on */
+	static void SetupInitialRenderTargetData(UTexture2D* InTextureSource, UTextureRenderTarget2D* InRenderTarget);
+
+	/** Tries to find Materials using the given Texture and retrieve the corresponding material indices from the MEsh Compon*/
+	static void FindMaterialIndicesUsingTexture(const UTexture* Texture, const UMeshComponent* MeshComponent, TArray<int32>& OutIndices);
+	
+	/** Retrieve LOD mesh sections from MeshComponent which use one of the given textures */
+	static void RetrieveMeshSectionsForTextures(const UMeshComponent* MeshComponent, int32 LODIndex, TArray<const UTexture*> Textures, TArray<FTexturePaintMeshSectionInfo>& OutSectionInfo);
+
+	/** Retrieve LOD mesh sections from MeshComponent which contain one of the Material Indices */
+	static void RetrieveMeshSectionsForMaterialIndices(const UMeshComponent* MeshComponent, int32 LODIndex, const TArray<int32>& MaterialIndices, TArray<FTexturePaintMeshSectionInfo>& OutSectionInfo);
+
+	/** Retrieves all Paintable Textures from the given MeshComponent */
+	static void RetrieveTexturesForComponent(const UMeshComponent* Component, IMeshPaintGeometryAdapter* Adapter, TArray<FPaintableTexture>& OutTextures);
+
+	/** Checks whether or not the given MeshComponent's materials use the specified texture */
+	static bool DoesMeshComponentUseTexture(UMeshComponent* MeshComponent, UTexture* Texture);
+};
\ No newline at end of file
diff --git a/Engine/Source/Editor/EditorStyle/Private/SlateEditorStyle.cpp b/Engine/Source/Editor/EditorStyle/Private/SlateEditorStyle.cpp
--- a/Engine/Source/Editor/EditorStyle/Private/SlateEditorStyle.cpp
+++ b/Engine/Source/Editor/EditorStyle/Private/SlateEditorStyle.cpp
@@ -2016,17 +2016,6 @@ void FSlateEditorStyle::FStyle::SetupGeneralStyles()
 #endif // WITH_EDITOR
 
 #if WITH_EDITOR || (IS_PROGRAM && WITH_UNREAL_DEVELOPER_TOOLS)
-	// Menu editor
-	{
-		Set("MultiBox.DragAbove", new BOX_BRUSH("Common/DropZoneIndicator_Above", FMargin(10.0f / 16.0f, 10.0f / 16.0f, 0, 0), SelectionColor));
-		Set("MultiBox.DragBelow", new BOX_BRUSH("Common/DropZoneIndicator_Below", FMargin(10.0f / 16.0f, 0, 0, 10.0f / 16.0f), SelectionColor));
-
-		Set("MultiBox.VisibleIcon16x", new IMAGE_BRUSH("Icons/icon_visible_16px", Icon16x16));
-		Set("MultiBox.VisibleHighlightIcon16x", new IMAGE_BRUSH("Icons/icon_visible_hi_16px", Icon16x16));
-		Set("MultiBox.NotVisibleIcon16x", new IMAGE_BRUSH("Icons/icon_invisible_16px", Icon16x16));
-		Set("MultiBox.NotVisibleHighlightIcon16x", new IMAGE_BRUSH("Icons/icon_invisible_hi_16px", Icon16x16));
-	}
-
 	// ToolBar
 	{
 		Set( "ToolBar.Background", new BOX_BRUSH( "Common/GroupBorder", FMargin(4.0f/16.0f) ) );
@@ -5379,23 +5368,36 @@ void FSlateEditorStyle::FStyle::SetupLevelEditorStyle()
 	// Mesh Paint
 	{
 		Set( "MeshPaint.Fill", new IMAGE_BRUSH( "/Icons/icon_MeshPaint_Fill_40x", Icon20x20) );
-		Set( "MeshPaint.Fill.Small", new IMAGE_BRUSH( "/Icons/icon_MeshPaint_Fill_40x", Icon16x16) );
+		Set( "MeshPaint.Fill.Small", new IMAGE_BRUSH( "/Icons/icon_MeshPaint_Fill_40x", Icon20x20) );
 		Set( "MeshPaint.Propagate", new IMAGE_BRUSH( "/Icons/icon_MatEd_Apply_40x", Icon20x20) );
-		Set( "MeshPaint.Propagate.Small", new IMAGE_BRUSH( "/Icons/icon_MatEd_Apply_40x", Icon16x16) );
+		Set( "MeshPaint.Propagate.Small", new IMAGE_BRUSH( "/Icons/icon_MatEd_Apply_40x", Icon20x20) );
 		Set( "MeshPaint.Import", new IMAGE_BRUSH( "/Icons/icon_Import_40x", Icon20x20) );
-		Set( "MeshPaint.Import.Small", new IMAGE_BRUSH( "/Icons/icon_Import_40x", Icon16x16) );
+		Set( "MeshPaint.Import.Small", new IMAGE_BRUSH( "/Icons/icon_Import_40x", Icon20x20) );
 		Set( "MeshPaint.FindInCB", new IMAGE_BRUSH( "/Icons/icon_toolbar_genericfinder_40px", Icon20x20) );
 		Set( "MeshPaint.Save", new IMAGE_BRUSH( "/Icons/icon_file_save_40x", Icon20x20) );
-		Set( "MeshPaint.Save.Small", new IMAGE_BRUSH( "/Icons/icon_file_save_40x", Icon16x16) );
+		Set( "MeshPaint.Save.Small", new IMAGE_BRUSH( "/Icons/icon_file_save_40x", Icon20x20) );
 		Set( "MeshPaint.Fix", new IMAGE_BRUSH( "/Icons/icon_tab_Toolbars_40x", Icon20x20) );
-		Set( "MeshPaint.Fix.Small", new IMAGE_BRUSH( "/Icons/icon_tab_Toolbars_40x", Icon16x16) );
+		Set( "MeshPaint.Fix.Small", new IMAGE_BRUSH( "/Icons/icon_tab_Toolbars_40x", Icon20x20) );
 		Set( "MeshPaint.Remove", new IMAGE_BRUSH("/Icons/Edit/icon_Edit_Delete_40x", Icon20x20));
-		Set( "MeshPaint.Remove.Small", new IMAGE_BRUSH("/Icons/Edit/icon_Edit_Delete_40x", Icon16x16));
+		Set( "MeshPaint.Remove.Small", new IMAGE_BRUSH("/Icons/Edit/icon_Edit_Delete_40x", Icon20x20));
 		Set( "MeshPaint.Copy", new IMAGE_BRUSH("/Icons/Edit/icon_Edit_Copy_40x", Icon20x20));
-		Set( "MeshPaint.Copy.Small", new IMAGE_BRUSH("/Icons/Edit/icon_Edit_Copy_40x", Icon16x16));
+		Set( "MeshPaint.Copy.Small", new IMAGE_BRUSH("/Icons/Edit/icon_Edit_Copy_40x", Icon20x20));
 		Set( "MeshPaint.Paste", new IMAGE_BRUSH("/Icons/Edit/icon_Edit_Paste_40x", Icon20x20));
-		Set( "MeshPaint.Paste.Small", new IMAGE_BRUSH("/Icons/Edit/icon_Edit_Paste_40x", Icon16x16));
+		Set( "MeshPaint.Paste.Small", new IMAGE_BRUSH("/Icons/Edit/icon_Edit_Paste_40x", Icon20x20));
 		Set( "MeshPaint.Swap", new IMAGE_BRUSH("/Icons/icon_MeshPaint_Swap_16x", Icon12x12));
+		Set("MeshPaint.PropagateTexturePaint", new IMAGE_BRUSH("/Icons/icon_MatEd_Apply_40x", Icon20x20));
+		Set("MeshPaint.PropagateTexturePaint.Small", new IMAGE_BRUSH("/Icons/icon_MatEd_Apply_40x", Icon20x20));
+		Set("MeshPaint.SaveTexturePaint", new IMAGE_BRUSH("/Icons/icon_file_save_40x", Icon20x20));
+		Set("MeshPaint.SaveTexturePaint.Small", new IMAGE_BRUSH("/Icons/icon_file_save_40x", Icon20x20));
+		Set("MeshPaint.PropagateVertexColorsToLODs", new IMAGE_BRUSH("Icons/icon_ViewMode_LODColoration_16px", Icon20x20));
+		Set("MeshPaint.PropagateVertexColorsToLODs.Small", new IMAGE_BRUSH("Icons/icon_ViewMode_LODColoration_16px", Icon20x20));
+		Set("MeshPaint.TexturePaint", new IMAGE_BRUSH("Icons/TexturePaint_40x", Icon20x20));
+		Set("MeshPaint.TexturePaint.Small", new IMAGE_BRUSH("Icons/TexturePaint_40x", Icon20x20));
+		Set("MeshPaint.ColorPaint", new IMAGE_BRUSH("Icons/VertexColorPaint_40x", Icon20x20));
+		Set("MeshPaint.ColorPaint.Small", new IMAGE_BRUSH("Icons/VertexColorPaint_40x", Icon20x20));
+		Set("MeshPaint.WeightPaint", new IMAGE_BRUSH("Icons/WeightPaint_40x", Icon20x20));
+		Set("MeshPaint.WeightPaint.Small", new IMAGE_BRUSH("Icons/WeightPaint_40x", Icon20x20));
+		Set("MeshPaint.Select", new IMAGE_BRUSH(TEXT("Icons/icon_Persona_MeshSectionSelection_40x"), Icon20x20));
 	}
 
 	// News Feed
diff --git a/Engine/Source/Editor/MeshPaintMode/Private/MeshPaintModeModule.cpp b/Engine/Source/Editor/MeshPaintMode/Private/MeshPaintModeModule.cpp
--- a/Engine/Source/Editor/MeshPaintMode/Private/MeshPaintModeModule.cpp
+++ b/Engine/Source/Editor/MeshPaintMode/Private/MeshPaintModeModule.cpp
@@ -14,25 +14,29 @@
 #include "PaintModeSettings.h"
 
 #include "Modules/ModuleManager.h"
+#include "Settings/LevelEditorMiscSettings.h"
 
 IMPLEMENT_MODULE(FMeshPaintModeModule, MeshPaintMode );
 
 void FMeshPaintModeModule::StartupModule()
 {
-	FEditorModeRegistry::Get().RegisterMode<FEdModeMeshPaint>(
-		FBuiltinEditorModes::EM_MeshPaint,
-		NSLOCTEXT("MeshPaint_Mode", "MeshPaint_ModeName", "Mesh Paint"),
-		FSlateIcon(FEditorStyle::GetStyleSetName(), "LevelEditor.MeshPaintMode", "LevelEditor.MeshPaintMode.Small"),
-		true, 200 );
-
-	/** Register detail/property customization */
-	FPropertyEditorModule& PropertyModule = FModuleManager::LoadModuleChecked<FPropertyEditorModule>("PropertyEditor");
-	PropertyModule.RegisterCustomClassLayout("PaintModeSettings", FOnGetDetailCustomizationInstance::CreateStatic(&FPaintModeSettingsCustomization::MakeInstance));
-	PropertyModule.RegisterCustomPropertyTypeLayout("VertexPaintSettings", FOnGetPropertyTypeCustomizationInstance::CreateStatic(&FVertexPaintSettingsCustomization::MakeInstance));
-	PropertyModule.RegisterCustomPropertyTypeLayout("TexturePaintSettings", FOnGetPropertyTypeCustomizationInstance::CreateStatic(&FTexturePaintSettingsCustomization::MakeInstance));
-	PropertyModule.RegisterCustomPropertyTypeLayout("TexturePaintSettings", FOnGetPropertyTypeCustomizationInstance::CreateStatic(&FTexturePaintSettingsCustomization::MakeInstance));
-
-	FModuleManager::Get().LoadModule("MeshPaint");
+	if (GetDefault<ULevelEditorMiscSettings>()->bEnableLegacyMeshPaintMode)
+	{
+		FEditorModeRegistry::Get().RegisterMode<FEdModeMeshPaint>(
+			FBuiltinEditorModes::EM_MeshPaint,
+			NSLOCTEXT("MeshPaint_Mode", "MeshPaint_ModeName", "Mesh Paint"),
+			FSlateIcon(FEditorStyle::GetStyleSetName(), "LevelEditor.MeshPaintMode", "LevelEditor.MeshPaintMode.Small"),
+			true, 200);
+
+		/** Register detail/property customization */
+		FPropertyEditorModule& PropertyModule = FModuleManager::LoadModuleChecked<FPropertyEditorModule>("PropertyEditor");
+		PropertyModule.RegisterCustomClassLayout("PaintModeSettings", FOnGetDetailCustomizationInstance::CreateStatic(&FPaintModeSettingsCustomization::MakeInstance));
+		PropertyModule.RegisterCustomPropertyTypeLayout("VertexPaintSettings", FOnGetPropertyTypeCustomizationInstance::CreateStatic(&FVertexPaintSettingsCustomization::MakeInstance));
+		PropertyModule.RegisterCustomPropertyTypeLayout("TexturePaintSettings", FOnGetPropertyTypeCustomizationInstance::CreateStatic(&FTexturePaintSettingsCustomization::MakeInstance));
+		PropertyModule.RegisterCustomPropertyTypeLayout("TexturePaintSettings", FOnGetPropertyTypeCustomizationInstance::CreateStatic(&FTexturePaintSettingsCustomization::MakeInstance));
+
+		FModuleManager::Get().LoadModule("MeshPaint");
+	}
 }
 
 void FMeshPaintModeModule::ShutdownModule()
diff --git a/Engine/Source/Editor/UnrealEd/Classes/Settings/LevelEditorMiscSettings.h b/Engine/Source/Editor/UnrealEd/Classes/Settings/LevelEditorMiscSettings.h
--- a/Engine/Source/Editor/UnrealEd/Classes/Settings/LevelEditorMiscSettings.h
+++ b/Engine/Source/Editor/UnrealEd/Classes/Settings/LevelEditorMiscSettings.h
@@ -46,6 +46,11 @@ public:
 	/** If enabled, the modes tab will revert to the previous mode UI */
 	UPROPERTY(EditAnywhere, config, AdvancedDisplay, Category = Editing, meta=(ConfigRestartRequired = true))
 	uint32 bEnableLegacyEditorModeUI : 1;
+
+	/** If enabled, the modes tab will revert to the original mesh paint */
+	UPROPERTY(EditAnywhere, config, AdvancedDisplay, Category = Editing, meta = (ConfigRestartRequired = true))
+	uint32 bEnableLegacyMeshPaintMode : 1;
+
 public:
 	/** If checked audio playing in the editor will continue to play even if the editor is in the background */
 	UPROPERTY(EditAnywhere, config, Category=Sound)
diff --git a/Engine/Source/Editor/UnrealEd/Private/Settings/SettingsClasses.cpp b/Engine/Source/Editor/UnrealEd/Private/Settings/SettingsClasses.cpp
--- a/Engine/Source/Editor/UnrealEd/Private/Settings/SettingsClasses.cpp
+++ b/Engine/Source/Editor/UnrealEd/Private/Settings/SettingsClasses.cpp
@@ -431,6 +431,7 @@ ULevelEditorMiscSettings::ULevelEditorMiscSettings( const FObjectInitializer& Ob
 	MinimumBoundsForCheckingSize = FVector(500.0f, 500.0f, 50.0f);
 	bCreateNewAudioDeviceForPlayInEditor = true;
 	bEnableLegacyEditorModeUI = false;
+	bEnableLegacyMeshPaintMode = true;
 }
 
 void ULevelEditorMiscSettings::PostEditChangeProperty( struct FPropertyChangedEvent& PropertyChangedEvent )
diff --git a/Engine/Source/Editor/UnrealEd/Private/Toolkits/BaseToolkit.cpp b/Engine/Source/Editor/UnrealEd/Private/Toolkits/BaseToolkit.cpp
--- a/Engine/Source/Editor/UnrealEd/Private/Toolkits/BaseToolkit.cpp
+++ b/Engine/Source/Editor/UnrealEd/Private/Toolkits/BaseToolkit.cpp
@@ -11,6 +11,7 @@
 #include "Editor.h"
 #include "InteractiveTool.h"
 #include "Tools/UEdMode.h"
+#include "Framework/MultiBox/MultiBoxBuilder.h"
 
 #define LOCTEXT_NAMESPACE "BaseToolkit"
 
@@ -215,10 +216,29 @@ void FModeToolkit::Init(const TSharedPtr< class IToolkitHost >& InitToolkitHost)
 	ToolkitMode = EToolkitMode::WorldCentric;
 	ToolkitHost = InitToolkitHost;
 
-	FToolkitManager::Get().RegisterNewToolkit( SharedThis( this ) );
+	if (UEdMode* ScriptableEditorMode = GetScriptableEditorMode())
+	{
+		ScriptableEditorMode->GetToolManager()->OnToolStarted.AddSP(this, &FModeToolkit::OnToolStarted);
+		ScriptableEditorMode->GetToolManager()->OnToolEnded.AddSP(this, &FModeToolkit::OnToolEnded);
+	}
 
 	FPropertyEditorModule& PropertyEditorModule = FModuleManager::GetModuleChecked<FPropertyEditorModule>("PropertyEditor");
 
+	FDetailsViewArgs ModeDetailsViewArgs(
+		/*bUpdateFromSelection=*/ false,
+		/*bLockable=*/ false,
+		/*bAllowSearch=*/ false,
+		FDetailsViewArgs::HideNameArea,
+		/*bHideSelectionTip=*/ true,
+		/*InNotifyHook=*/ nullptr,
+		/*InSearchInitialKeyFocus=*/ false,
+		/*InViewIdentifier=*/ NAME_None);
+	ModeDetailsViewArgs.DefaultsOnlyVisibility = EEditDefaultsOnlyNodeVisibility::Automatic;
+	ModeDetailsViewArgs.bShowOptions = false;
+	ModeDetailsViewArgs.bAllowMultipleTopLevelObjects = true;
+
+	ModeDetailsView = PropertyEditorModule.CreateDetailView(ModeDetailsViewArgs);
+
 	FDetailsViewArgs DetailsViewArgs(
 		/*bUpdateFromSelection=*/ false,
 		/*bLockable=*/ false,
@@ -233,9 +253,24 @@ void FModeToolkit::Init(const TSharedPtr< class IToolkitHost >& InitToolkitHost)
 	DetailsViewArgs.bAllowMultipleTopLevelObjects = true;
 
 	DetailsView = PropertyEditorModule.CreateDetailView(DetailsViewArgs);
+
+
+	FToolkitManager::Get().RegisterNewToolkit(SharedThis(this));
 }
 
 
+FModeToolkit::~FModeToolkit()
+{
+	if (UEdMode* ScriptableEditorMode = GetScriptableEditorMode())
+	{
+		if (ScriptableEditorMode->GetToolManager())
+		{
+			ScriptableEditorMode->GetToolManager()->OnToolStarted.RemoveAll(this);
+			ScriptableEditorMode->GetToolManager()->OnToolEnded.RemoveAll(this);
+		}
+	}
+}
+
 FName FModeToolkit::GetToolkitFName() const
 {
 	return FName("EditorModeToolkit");
@@ -292,33 +327,19 @@ bool FModeToolkit::CanCompleteActiveTool()
 }
 
 
-FReply FModeToolkit::StartTool(const FString& ToolTypeIdentifier)
+void FModeToolkit::OnToolStarted(UInteractiveToolManager* Manager, UInteractiveTool* Tool)
 {
-	if (GetScriptableEditorMode()->GetToolManager()->SelectActiveToolType(EToolSide::Left, ToolTypeIdentifier) == false)
-	{
-		UE_LOG(LogTemp, Warning, TEXT("ToolManager: Unknown Tool Type %s"), *ToolTypeIdentifier);
-	}
-	else
+	// Update properties panel
+	UInteractiveTool* CurTool = GetScriptableEditorMode()->GetToolManager()->GetActiveTool(EToolSide::Left);
+	if (CurTool)
 	{
-		UE_LOG(LogTemp, Warning, TEXT("ToolManager: Starting Tool Type %s"), *ToolTypeIdentifier);
-		GetScriptableEditorMode()->GetToolManager()->ActivateTool(EToolSide::Left);
-
-		// Update properties panel
-		UInteractiveTool* CurTool = GetScriptableEditorMode()->GetToolManager()->GetActiveTool(EToolSide::Left);
 		DetailsView->SetObjects(CurTool->GetToolProperties());
 	}
-	return FReply::Handled();
 }
 
-FReply FModeToolkit::EndTool(EToolShutdownType ShutdownType)
+void FModeToolkit::OnToolEnded(UInteractiveToolManager* Manager, UInteractiveTool* Tool)
 {
-	UE_LOG(LogTemp, Warning, TEXT("ENDING TOOL"));
-
-	GetScriptableEditorMode()->GetToolManager()->DeactivateTool(EToolSide::Left, ShutdownType);
-
 	DetailsView->SetObject(nullptr);
-
-	return FReply::Handled();
 }
 
 class FEdMode* FModeToolkit::GetEditorMode() const
@@ -329,4 +350,63 @@ class FEdMode* FModeToolkit::GetEditorMode() const
 UEdMode* FModeToolkit::GetScriptableEditorMode() const
 {
 	return nullptr;
-}
\ No newline at end of file
+}
+
+TSharedPtr<SWidget> FModeToolkit::GetInlineContent() const
+{
+	return SNew(SVerticalBox)
+		+SVerticalBox::Slot()
+		.HAlign(HAlign_Center)
+		.AutoHeight()
+		[
+			SNew(STextBlock)
+			.Text(GetBaseToolkitName())
+		]
+		+ SVerticalBox::Slot()
+		.AutoHeight()
+		[
+			ModeDetailsView.ToSharedRef()
+		]
+		+ SVerticalBox::Slot()
+		.AutoHeight()
+		[
+			DetailsView.ToSharedRef()
+		];
+}
+
+void FModeToolkit::BuildToolPalette(FName PaletteName, class FToolBarBuilder& ToolbarBuilder)
+{
+	if (UEdMode* EditorMode = GetScriptableEditorMode())
+	{
+		TMap<FName, TArray<TSharedPtr<FUICommandInfo>>> CommandLists = EditorMode->GetModeCommands();
+		TArray<TSharedPtr<FUICommandInfo>>* CurrentCommandListPtr = CommandLists.Find(PaletteName);
+		if (CurrentCommandListPtr)
+		{
+			TArray<TSharedPtr<FUICommandInfo>> CurrentCommandList = *CurrentCommandListPtr;
+			for (TSharedPtr<FUICommandInfo> Command : CurrentCommandList)
+			{
+				ToolbarBuilder.AddToolBarButton(Command);
+			}
+			
+		}
+
+	}
+}
+
+void FModeToolkit::OnToolPaletteChanged(FName PaletteName)
+{
+	if (UEdMode* EditorMode = GetScriptableEditorMode())
+	{
+		EditorMode->SetCurrentPaletteName(PaletteName);
+		if (UInteractiveTool* CurrentTool = EditorMode->GetToolManager()->GetActiveTool(EToolSide::Left))
+		{
+			// When swapping tools, default to accepting the active tool
+			EditorMode->GetToolManager()->DeactivateTool(EToolSide::Left, EToolShutdownType::Accept);
+		}
+	}
+}
+
+void FModeToolkit::SetModeSettingsObject(UObject* InSettingsObject)
+{
+	ModeDetailsView->SetObject(InSettingsObject);
+}
diff --git a/Engine/Source/Editor/UnrealEd/Private/Tools/UEdMode.cpp b/Engine/Source/Editor/UnrealEd/Private/Tools/UEdMode.cpp
--- a/Engine/Source/Editor/UnrealEd/Private/Tools/UEdMode.cpp
+++ b/Engine/Source/Editor/UnrealEd/Private/Tools/UEdMode.cpp
@@ -23,6 +23,8 @@
 #include "LevelEditorViewport.h"
 #include "InteractiveToolManager.h"
 #include "InteractiveToolObjects.h"
+#include "Editor/EditorEngine.h"
+#include "GameFramework/Actor.h"
 
 
 class FEdModeQueriesImpl : public IToolsContextQueriesAPI
@@ -85,8 +87,11 @@ public:
 		StateOut.ToolManager = ToolsContext->ToolManager;
 		StateOut.GizmoManager = ToolsContext->GizmoManager;
 		StateOut.World = GEditor->GetWorld();
-		GEditor->GetSelectedActors()->GetSelectedObjects(StateOut.SelectedActors);
-		GEditor->GetSelectedComponents()->GetSelectedObjects(StateOut.SelectedComponents);
+		USelection* SelectionActorSet = GEditor->GetSelectedActors();
+		SelectionActorSet->GetSelectedObjects(StateOut.SelectedActors);
+		USelection* SelectionComponentSet = GEditor->GetSelectedComponents();
+		SelectionComponentSet->GetSelectedObjects(StateOut.SelectedComponents);
+
 	}
 
 
@@ -378,6 +383,9 @@ UEdMode::UEdMode()
 {
 	bDrawKillZ = true;
 	ToolsContext = nullptr;
+	ToolsContextClass = UInteractiveToolsContext::StaticClass();
+	ToolCommandList = MakeShareable(new FUICommandList);
+	bCheckIfDefaultToolNeeded = true;
 }
 
 
@@ -412,8 +420,6 @@ FRay UEdMode::GetRayFromMousePos(FEditorViewportClient* ViewportClient, FViewpor
 	return FRay(MouseViewportRay.GetOrigin(), MouseViewportRay.GetDirection(), true);
 }
 
-
-
 bool UEdMode::MouseEnter(FEditorViewportClient* ViewportClient, FViewport* Viewport, int32 x, int32 y)
 {
 	CurrentMouseState.Mouse.Position2D = FVector2D(x, y);
@@ -625,6 +631,14 @@ bool UEdMode::InputKey(FEditorViewportClient* ViewportClient, FViewport* Viewpor
 
 	if(!bHandled)
 	{
+
+		if (Event != IE_Released)
+		{
+			if (ToolCommandList->ProcessCommandBindings(Key, FSlateApplication::Get().GetModifierKeys(), false/*Event == IE_Repeat*/))
+			{
+				return true;
+			}
+		}
 		// Next pass input to the mode toolkit
 		if (Toolkit.IsValid() && ((Event == IE_Pressed) || (Event == IE_Repeat)))
 		{
@@ -645,7 +659,7 @@ bool UEdMode::InputKey(FEditorViewportClient* ViewportClient, FViewport* Viewpor
 		}
 	}
 
-	return false;
+	return bHandled;
 }
 
 bool UEdMode::InputAxis(FEditorViewportClient* InViewportClient, FViewport* Viewport, int32 ControllerId, FKey Key, float Delta, float DeltaTime)
@@ -696,6 +710,11 @@ void UEdMode::Tick(FEditorViewportClient* ViewportClient, float DeltaTime)
 		ToolsContext->ToolManager->Tick(DeltaTime);
 		ToolsContext->GizmoManager->Tick(DeltaTime);
 
+		if(bCheckIfDefaultToolNeeded && ToolsContext->ToolManager->GetActiveTool(EToolSide::Left) == nullptr)
+		{
+			ActivateDefaultTool();
+		}
+
 		if (bInvalidationPending)
 		{
 			ViewportClient->Invalidate();
@@ -712,13 +731,10 @@ void UEdMode::Tick(FEditorViewportClient* ViewportClient, float DeltaTime)
 	}
 }
 
-void UEdMode::ActorSelectionChangeNotify()
-{
-}
 
 bool UEdMode::HandleClick(FEditorViewportClient* InViewportClient, HHitProxy *HitProxy, const FViewportClick &Click)
 {
-	return false;
+	return true;
 }
 
 void UEdMode::Enter()
@@ -733,23 +749,32 @@ void UEdMode::Enter()
 
 	bPendingDeletion = false;
 
-	if (!Toolkit.IsValid())
-	{
-		Toolkit = MakeShareable(new FModeToolkit);
-		Toolkit->Init(Owner->GetToolkitHost());
-	}
-
 	// initialize the adapter that attaches the ToolsContext to this FEdMode
-	ToolsContext = NewObject<UInteractiveToolsContext>(GetTransientPackage(), TEXT("ToolsContext"), RF_Transient);
+	ToolsContext = NewObject<UInteractiveToolsContext>(GetTransientPackage(), ToolsContextClass.Get(), TEXT("ToolsContext"), RF_Transient);
 	TransactionAPI = new FEdModeTransactionImpl(ToolsContext, this);
 	QueriesAPI = new FEdModeQueriesImpl(ToolsContext, this);
 	AssetAPI = new FEditorToolAssetAPI();
 	ToolsContext->Initialize(QueriesAPI, TransactionAPI);
 
+	ToolsContext->OnToolNotificationMessage.AddLambda([this](const FText& Message)
+	{
+		OnToolNotificationMessage.Broadcast(Message);
+	});
+	ToolsContext->OnToolWarningMessage.AddLambda([this](const FText& Message)
+	{
+		OnToolWarningMessage.Broadcast(Message);
+	});
+
+
+	GetToolManager()->OnToolStarted.AddUObject(this, &UEdMode::OnToolStarted);
+	GetToolManager()->OnToolEnded.AddUObject(this, &UEdMode::OnToolEnded);
+
 	// enable auto invalidation in Editor, because invalidating for all hover and capture events is unpleasant
 	ToolsContext->InputRouter->bAutoInvalidateOnHover = true;
 	ToolsContext->InputRouter->bAutoInvalidateOnCapture = true;
 
+	// Now that the context is ready, make the toolkit
+	CreateToolkit();
 
 	// set up standard materials
 	StandardVertexColorMaterial = LoadObject<UMaterial>(nullptr, TEXT("/Game/Materials/VertexColor"));
@@ -767,13 +792,25 @@ void UEdMode::Enter()
 	FEditorDelegates::EditorModeIDEnter.Broadcast(GetID());
 	const bool bIsEnteringMode = true;
 	Owner->BroadcastEditorModeIDChanged(GetID(), bIsEnteringMode);
+
+}
+
+void UEdMode::RegisterTool(TSharedPtr<FUICommandInfo> UICommand, FString ToolIdentifier, UInteractiveToolBuilder* Builder)
+{
+	const TSharedRef<FUICommandList>& CommandList = Toolkit->GetToolkitCommands();
+	ToolsContext->ToolManager->RegisterToolType(ToolIdentifier, Builder);
+	CommandList->MapAction(UICommand,
+		FExecuteAction::CreateLambda([this, ToolIdentifier]() { ToolsContext->StartTool(EToolSide::Mouse, ToolIdentifier); }),
+		FCanExecuteAction::CreateLambda([this, ToolIdentifier]() { return GetToolManager()->CanActivateTool(EToolSide::Mouse, ToolIdentifier); }),
+		FIsActionChecked::CreateLambda([this, Builder]() { return ToolsContext->IsToolBuilderActive(EToolSide::Mouse, Builder); }));
 }
 
 void UEdMode::Exit()
 {
 	FEditorDelegates::BeginPIE.Remove(BeginPIEDelegateHandle);
 	FEditorDelegates::PreSaveWorld.Remove(PreSaveWorldDelegateHandle);
-
+	GetToolManager()->OnToolStarted.RemoveAll(this);
+	GetToolManager()->OnToolEnded.RemoveAll(this);
 	// auto-accept any in-progress tools
 	DeactivateAllActiveTools();
 
@@ -801,6 +838,9 @@ void UEdMode::Exit()
 		Toolkit.Reset();
 	}
 
+	OnToolNotificationMessage.Clear();
+	OnToolWarningMessage.Clear();
+
 	const bool bIsEnteringMode = false;
 	Owner->BroadcastEditorModeIDChanged(GetID(), bIsEnteringMode);
 	FEditorDelegates::EditorModeIDExit.Broadcast(GetID());
@@ -988,7 +1028,10 @@ AActor* UEdMode::GetFirstSelectedActorInstance() const
 
 void UEdMode::DeactivateAllActiveTools()
 {
-	ToolsContext->DeactivateAllActiveTools();
+	if (ToolsContext)
+	{
+		ToolsContext->DeactivateAllActiveTools();
+	}
 	RestoreEditorState();
 }
 
@@ -997,6 +1040,22 @@ UInteractiveToolManager* UEdMode::GetToolManager() const
 	return ToolsContext->ToolManager;
 }
 
+void UEdMode::CreateToolkit()
+{
+	if (!Toolkit.IsValid())
+	{
+		Toolkit = MakeShareable(new FModeToolkit);
+		Toolkit->Init(Owner->GetToolkitHost());
+	}
+
+	UClass* LoadedSettingsObject = SettingsClass.LoadSynchronous();
+	SettingsObject = NewObject<UObject>(this, LoadedSettingsObject);
+	if (SettingsObject)
+	{
+		Toolkit->SetModeSettingsObject(SettingsObject);
+	}
+}
+
 bool UEdMode::IsSnapRotationEnabled()
 {
 	return GetDefault<ULevelEditorViewportSettings>()->RotGridEnabled;
diff --git a/Engine/Source/Editor/UnrealEd/Public/Toolkits/BaseToolkit.h b/Engine/Source/Editor/UnrealEd/Public/Toolkits/BaseToolkit.h
--- a/Engine/Source/Editor/UnrealEd/Public/Toolkits/BaseToolkit.h
+++ b/Engine/Source/Editor/UnrealEd/Public/Toolkits/BaseToolkit.h
@@ -10,6 +10,7 @@
 #include "Toolkits/IToolkitHost.h"
 
 class UInteractiveTool;
+class UInteractiveToolManager;
 class SDockableTab;
 class UEdMode;
 class IDetailsView;
@@ -63,6 +64,7 @@ public:
 	 */
 	void AddToolkitTab(const TSharedRef<SDockableTab>& TabToAdd, const EToolkitTabSpot::Type TabSpot);
 
+
 protected:
 
 	/** @return Returns the prefix string to use for tabs created for this toolkit.  In world-centric mode, tabs get a
@@ -106,6 +108,7 @@ public:
 
 	/** Initializes the mode toolkit */
 	virtual void Init(const TSharedPtr<IToolkitHost>& InitToolkitHost);
+	~FModeToolkit();
 
 public:
 
@@ -126,21 +129,24 @@ public:
 	virtual const TArray<UObject*>* GetObjectsCurrentlyBeingEdited() const override;
 	virtual FLinearColor GetWorldCentricTabColorScale() const override;
 	virtual class FEdMode* GetEditorMode() const override;
-
 	virtual UEdMode* GetScriptableEditorMode() const;
+	virtual TSharedPtr<SWidget> GetInlineContent() const override;
+	virtual void BuildToolPalette(FName PaletteName, class FToolBarBuilder& ToolbarBuilder) override;
+	virtual void OnToolPaletteChanged(FName PaletteName) override;
 
-private:
-
-	TSharedPtr<SWidget> ToolkitWidget;
-	TSharedPtr<IDetailsView> DetailsView;
-
-	// these functions just forward calls to the ToolsContext / ToolManager
+	void SetModeSettingsObject(UObject* InSettingsObject);
 
+protected:
 	bool CanStartTool(const FString& ToolTypeIdentifier);
 	bool CanAcceptActiveTool();
 	bool CanCancelActiveTool();
 	bool CanCompleteActiveTool();
 
-	FReply StartTool(const FString& ToolTypeIdentifier);
-	FReply EndTool(EToolShutdownType ShutdownType);
+	virtual void OnToolStarted(UInteractiveToolManager* Manager, UInteractiveTool* Tool);
+	virtual void OnToolEnded(UInteractiveToolManager* Manager, UInteractiveTool* Tool);
+
+protected:
+	TSharedPtr<SWidget> ToolkitWidget;
+	TSharedPtr<IDetailsView> ModeDetailsView;
+	TSharedPtr<IDetailsView> DetailsView;
 };
diff --git a/Engine/Source/Editor/UnrealEd/Public/Tools/UEdMode.h b/Engine/Source/Editor/UnrealEd/Public/Tools/UEdMode.h
--- a/Engine/Source/Editor/UnrealEd/Public/Tools/UEdMode.h
+++ b/Engine/Source/Editor/UnrealEd/Public/Tools/UEdMode.h
@@ -11,6 +11,7 @@
 #include "UObject/Object.h"
 #include "Math/Ray.h"
 #include "InputState.h"
+#include "UObject/SoftObjectPtr.h"
 #include "UEdMode.generated.h"
 
 
@@ -27,6 +28,7 @@ struct FConvexVolume;
 struct FViewportClick;
 class UInteractiveToolsContext;
 class UInteractiveToolManager;
+class UInteractiveTool;
 class IToolsContextQueriesAPI;
 class IToolsContextTransactionsAPI;
 class IToolsContextAssetAPI;
@@ -38,6 +40,9 @@ struct FViewportClick;
 class FModeTool;
 class FEditorViewportClient;
 struct FViewportClick;
+class UInteractiveToolBuilder;
+class FUICommandInfo;
+class FUICommandList;
 
 
 /** Outcomes when determining whether it's possible to perform an action on the edit modes*/
@@ -57,7 +62,7 @@ namespace EEditAction
 /**
  * Base class for all editor modes.
  */
-UCLASS()
+UCLASS(Abstract)
 class UNREALED_API UEdMode : public UObject, public FEditorCommonDrawHelper
 {
 	GENERATED_BODY()
@@ -117,7 +122,7 @@ public:
 	virtual bool IsCompatibleWith(FEditorModeID OtherModeID) const { return false; }
 	virtual void ActorMoveNotify() {}
 	virtual void ActorsDuplicatedNotify(TArray<AActor*>& PreDuplicateSelection, TArray<AActor*>& PostDuplicateSelection, bool bOffsetLocations) {}
-	virtual void ActorSelectionChangeNotify();
+	virtual void ActorSelectionChangeNotify() {};
 	virtual void ActorPropChangeNotify() {}
 	virtual void MapChangeNotify() {}
 
@@ -164,6 +169,10 @@ public:
 	virtual void UpdateInternalData() {}
 
 	virtual void Enter();
+
+	virtual void RegisterTool(TSharedPtr<FUICommandInfo> UICommand, FString ToolIdentifier, UInteractiveToolBuilder* Builder);
+
+
 	virtual void Exit();
 	virtual UTexture2D* GetVertexTexture() { return GEngine->DefaultBSPVertexTexture; }
 
@@ -232,13 +241,6 @@ public:
 	/** Returns the owning mode manager for this mode */
 	class FEditorModeTools* GetModeManager() const;
 
-	/**
-	 * Called when the editor mode should rebuild its toolbar
-	 *
-	 * @param ToolbarBuilder	The builder which should be used to add toolbar widgets
-	 */
-	virtual void BuildModeToolbar(class FToolBarBuilder& ToolbarBuilder) {}
-
 public:
 
 	/** Request that this mode be deleted at the next convenient opportunity (FEditorModeTools::Tick) */
@@ -261,6 +263,7 @@ private:
 	virtual void TerminateActiveToolsOnSaveWorld();
 
 	FRay GetRayFromMousePos(FEditorViewportClient* ViewportClient, FViewport* Viewport, int MouseX, int MouseY);
+	
 protected:
 
 	/** Information pertaining to this mode. Assigned by FEditorModeRegistry. */
@@ -299,10 +302,31 @@ public:
 	 */
 	virtual UInteractiveToolManager* GetToolManager() const;
 
+	virtual TMap<FName, TArray<TSharedPtr<FUICommandInfo>>> GetModeCommands() const
+	{
+		return TMap<FName, TArray<TSharedPtr<FUICommandInfo>>>();
+	};
+
+	void SetCurrentPaletteName(FName InName)
+	{
+		CurrentPaletteName = InName;
+		UpdateOnPaletteChange();
+	};
+
 protected:
+	virtual void CreateToolkit();
+	virtual void OnToolStarted(UInteractiveToolManager* Manager, UInteractiveTool* Tool) {};
+	virtual void OnToolEnded(UInteractiveToolManager* Manager, UInteractiveTool* Tool) { bCheckIfDefaultToolNeeded = true; };
+	virtual void ActivateDefaultTool() {};
+	virtual void UpdateOnPaletteChange() {};
 
-	UInteractiveToolsContext* ToolsContext;
+protected:
 
+	UPROPERTY()
+	UInteractiveToolsContext* ToolsContext;
+	TSoftClassPtr<UInteractiveToolsContext> ToolsContextClass;
+	/** Command list lives here so that the key bindings on the commands can be processed in the viewport. */
+	TSharedPtr<FUICommandList> ToolCommandList;
 	bool bInvalidationPending;
 	IToolsContextQueriesAPI* QueriesAPI;
 	IToolsContextTransactionsAPI* TransactionAPI;
@@ -313,5 +337,17 @@ protected:
 	// called before a Save starts. This currently shuts down active tools.
 	FDelegateHandle PreSaveWorldDelegateHandle;
 
+	UPROPERTY()
+	TSoftClassPtr<UObject> SettingsClass;
+
+	UPROPERTY(Transient)
+	UObject* SettingsObject;
+
+	DECLARE_MULTICAST_DELEGATE_OneParam(FOnModeToolNotification, const FText&);
+	FOnModeToolNotification OnToolNotificationMessage;
+	FOnModeToolNotification OnToolWarningMessage;
+
+	FName CurrentPaletteName;
+	bool bCheckIfDefaultToolNeeded;
 };
 
diff --git a/Engine/Source/Runtime/Experimental/InteractiveToolsFramework/Private/BaseBehaviors/SingleClickBehavior.cpp b/Engine/Source/Runtime/Experimental/InteractiveToolsFramework/Private/BaseBehaviors/SingleClickBehavior.cpp
--- a/Engine/Source/Runtime/Experimental/InteractiveToolsFramework/Private/BaseBehaviors/SingleClickBehavior.cpp
+++ b/Engine/Source/Runtime/Experimental/InteractiveToolsFramework/Private/BaseBehaviors/SingleClickBehavior.cpp
@@ -29,20 +29,22 @@ FInputCaptureRequest USingleClickInputBehavior::WantsCapture(const FInputDeviceS
 }
 
 
-FInputCaptureUpdate USingleClickInputBehavior::BeginCapture(const FInputDeviceState& input, EInputCaptureSide eSide)
+FInputCaptureUpdate USingleClickInputBehavior::BeginCapture(const FInputDeviceState& Input, EInputCaptureSide eSide)
 {
+	Modifiers.UpdateModifiers(Input, Target);
 	return FInputCaptureUpdate::Begin(this, EInputCaptureSide::Any);
 }
 
 
-FInputCaptureUpdate USingleClickInputBehavior::UpdateCapture(const FInputDeviceState& input, const FInputCaptureData& data)
+FInputCaptureUpdate USingleClickInputBehavior::UpdateCapture(const FInputDeviceState& Input, const FInputCaptureData& Data)
 {
-	if (IsReleased(input)) 
+	Modifiers.UpdateModifiers(Input, Target);
+	if (IsReleased(Input))
 	{
 		if (HitTestOnRelease == false || 
-			Target->IsHitByClick(GetDeviceRay(input)) )
+			Target->IsHitByClick(GetDeviceRay(Input)) )
 		{
-			Clicked(input, data);
+			Clicked(Input, Data);
 		}
 
 		return FInputCaptureUpdate::End();
diff --git a/Engine/Source/Runtime/Experimental/InteractiveToolsFramework/Private/BaseTools/BaseBrushTool.cpp b/Engine/Source/Runtime/Experimental/InteractiveToolsFramework/Private/BaseTools/BaseBrushTool.cpp
--- a/Engine/Source/Runtime/Experimental/InteractiveToolsFramework/Private/BaseTools/BaseBrushTool.cpp
+++ b/Engine/Source/Runtime/Experimental/InteractiveToolsFramework/Private/BaseTools/BaseBrushTool.cpp
@@ -18,6 +18,7 @@ UBrushBaseProperties::UBrushBaseProperties()
 	bSpecifyRadius = false;
 	BrushRadius = 10.0f;
 	BrushStrength = 0.5f;
+	BrushFalloffAmount = 1.0f;
 }
 
 void UBrushBaseProperties::SaveProperties(UInteractiveTool* SaveFromTool)
@@ -26,6 +27,8 @@ void UBrushBaseProperties::SaveProperties(UInteractiveTool* SaveFromTool)
 	PropertyCache->BrushSize = this->BrushSize;
 	PropertyCache->bSpecifyRadius = this->bSpecifyRadius;
 	PropertyCache->BrushRadius = this->BrushRadius;
+	PropertyCache->BrushFalloffAmount = this->BrushFalloffAmount;
+	PropertyCache->BrushStrength = this->BrushStrength;
 }
 
 void UBrushBaseProperties::RestoreProperties(UInteractiveTool* RestoreToTool)
@@ -34,6 +37,8 @@ void UBrushBaseProperties::RestoreProperties(UInteractiveTool* RestoreToTool)
 	this->BrushSize = PropertyCache->BrushSize;
 	this->bSpecifyRadius = PropertyCache->bSpecifyRadius;
 	this->BrushRadius = PropertyCache->BrushRadius;
+	this->BrushFalloffAmount = PropertyCache->BrushFalloffAmount;
+	this->BrushStrength = PropertyCache->BrushStrength;
 }
 
 
@@ -41,16 +46,17 @@ void UBrushBaseProperties::RestoreProperties(UInteractiveTool* RestoreToTool)
 
 UBaseBrushTool::UBaseBrushTool()
 {
+	PropertyClass = UBrushBaseProperties::StaticClass();
 }
 
 
 void UBaseBrushTool::Setup()
 {
 	UMeshSurfacePointTool::Setup();
-
+	BrushProperties = NewObject<UBrushBaseProperties>(this, PropertyClass.Get(), TEXT("Brush"));
 	double MaxDimension = EstimateMaximumTargetDimension();
 	BrushRelativeSizeRange = TInterval<float>(MaxDimension*0.01, MaxDimension);
-	BrushProperties = NewObject<UBrushBaseProperties>(this, TEXT("Brush"));
+
 	RecalculateBrushRadius();
 
 	// initialize our properties
@@ -89,21 +95,86 @@ void UBaseBrushTool::DecreaseBrushSizeAction()
 }
 
 
+void UBaseBrushTool::IncreaseBrushStrengthAction()
+{
+	const float ChangeAmount = 0.02f;
+	const float OldValue = BrushProperties->BrushStrength;
+
+	float NewValue = OldValue + ChangeAmount;
+	BrushProperties->BrushStrength = FMath::Clamp(NewValue, 0.f, 1.f);
+}
+
+void UBaseBrushTool::DecreaseBrushStrengthAction()
+{
+	const float ChangeAmount = 0.02f;
+	const float OldValue = BrushProperties->BrushStrength;
+
+	float NewValue = OldValue - ChangeAmount;
+	BrushProperties->BrushStrength = FMath::Clamp(NewValue, 0.f, 1.f);
+}
+
+
+void UBaseBrushTool::IncreaseBrushFalloffAction()
+{
+	const float ChangeAmount = 0.02f;
+	const float OldValue = BrushProperties->BrushFalloffAmount;
+
+	float NewValue = OldValue + ChangeAmount;
+	BrushProperties->BrushFalloffAmount = FMath::Clamp(NewValue, 0.f, 1.f);
+}
+
+void UBaseBrushTool::DecreaseBrushFalloffAction()
+{
+	const float ChangeAmount = 0.02f;
+	const float OldValue = BrushProperties->BrushFalloffAmount;
+
+	float NewValue = OldValue - ChangeAmount;
+	BrushProperties->BrushFalloffAmount = FMath::Clamp(NewValue, 0.f, 1.f);
+}
+
 void UBaseBrushTool::RegisterActions(FInteractiveToolActionSet& ActionSet)
 {
-	ActionSet.RegisterAction(this, (int32)EStandardToolActions::BaseClientDefinedActionID+12,
+	ActionSet.RegisterAction(this, (int32)EStandardToolActions::BaseClientDefinedActionID + 10,
 		TEXT("BrushIncreaseSize"), 
 		LOCTEXT("BrushIncreaseSize", "Increase Brush Size"),
-		LOCTEXT("BrushIncreaseSizeTooltip", "Increase size of brush"),
+		LOCTEXT("BrushIncreaseSizeTooltip", "Press this key to increase brush radius by a percentage of its current size."),
 		EModifierKey::None, EKeys::RightBracket,
 		[this]() { IncreaseBrushSizeAction(); });
 
-	ActionSet.RegisterAction(this, (int32)EStandardToolActions::BaseClientDefinedActionID + 37,
+	ActionSet.RegisterAction(this, (int32)EStandardToolActions::BaseClientDefinedActionID + 11,
 		TEXT("BrushDecreaseSize"), 
 		LOCTEXT("BrushDecreaseSize", "Decrease Brush Size"),
-		LOCTEXT("BrushDecreaseSizeTooltip", "Decrease size of brush"),
+		LOCTEXT("BrushDecreaseSizeTooltip", "Press this key to decrease brush radius by a percentage of its current size."),
 		EModifierKey::None, EKeys::LeftBracket,
 		[this]() { DecreaseBrushSizeAction(); });
+
+	ActionSet.RegisterAction(this, (int32)EStandardToolActions::BaseClientDefinedActionID + 12,
+		TEXT("BrushIncreaseFalloff"),
+		LOCTEXT("BrushIncreaseFalloff", "Increase Brush Falloff"),
+		LOCTEXT("BrushIncreaseFalloffTooltip", "Press this key to increase brush falloff by a fixed increment."),
+		EModifierKey::Shift | EModifierKey::Control, EKeys::RightBracket,
+		[this]() { IncreaseBrushFalloffAction(); });
+
+	ActionSet.RegisterAction(this, (int32)EStandardToolActions::BaseClientDefinedActionID + 13,
+		TEXT("BrushDecreaseFalloff"),
+		LOCTEXT("BrushDecreaseFalloff", "Decrease Brush Falloff"),
+		LOCTEXT("BrushDecreaseFalloffTooltip", "Press this key to decrease brush falloff by a fixed increment."),
+		EModifierKey::Shift | EModifierKey::Control, EKeys::LeftBracket,
+		[this]() { DecreaseBrushFalloffAction(); });
+
+	ActionSet.RegisterAction(this, (int32)EStandardToolActions::BaseClientDefinedActionID + 14,
+		TEXT("BrushIncreaseStrength"),
+		LOCTEXT("BrushIncreaseStrength", "Increase Brush Strength"),
+		LOCTEXT("BrushIncreaseStrengthTooltip", "Press this key to increase brush strength by a fixed increment."),
+		EModifierKey::Control, EKeys::RightBracket,
+		[this]() { IncreaseBrushStrengthAction(); });
+
+	ActionSet.RegisterAction(this, (int32)EStandardToolActions::BaseClientDefinedActionID + 15,
+		TEXT("BrushDecreaseStrength"),
+		LOCTEXT("BrushDecreaseStrength", "Decrease Brush Strength"),
+		LOCTEXT("BrushDecreaseStrengthTooltip", "Press this key to decrease brush strength by a fixed increment."),
+		EModifierKey::Control, EKeys::LeftBracket,
+		[this]() { DecreaseBrushStrengthAction(); });
 }
 
 
@@ -132,6 +203,7 @@ void UBaseBrushTool::OnBeginDrag(const FRay& Ray)
 		LastBrushStamp.WorldPosition = OutHit.ImpactPoint;
 		LastBrushStamp.WorldNormal = OutHit.Normal;
 		LastBrushStamp.HitResult = OutHit;
+		LastBrushStamp.Falloff = BrushProperties->BrushFalloffAmount;
 	}
 	bInBrushStroke = true;
 }
@@ -145,6 +217,7 @@ void UBaseBrushTool::OnUpdateDrag(const FRay& Ray)
 		LastBrushStamp.WorldPosition = OutHit.ImpactPoint;
 		LastBrushStamp.WorldNormal = OutHit.Normal;
 		LastBrushStamp.HitResult = OutHit;
+		LastBrushStamp.Falloff = BrushProperties->BrushFalloffAmount;
 	}
 }
 
@@ -163,6 +236,7 @@ bool UBaseBrushTool::OnUpdateHover(const FInputDeviceRay& DevicePos)
 		LastBrushStamp.WorldPosition = OutHit.ImpactPoint;
 		LastBrushStamp.WorldNormal = OutHit.Normal;
 		LastBrushStamp.HitResult = OutHit;
+		LastBrushStamp.Falloff = BrushProperties->BrushFalloffAmount;
 	}
 	return true;
 }
@@ -195,7 +269,7 @@ void UBaseBrushTool::SetupBrushStampIndicator()
 
 void UBaseBrushTool::UpdateBrushStampIndicator()
 {
-	BrushStampIndicator->Update(LastBrushStamp.Radius, LastBrushStamp.WorldPosition, LastBrushStamp.WorldNormal);
+	BrushStampIndicator->Update(LastBrushStamp.Radius, LastBrushStamp.WorldPosition, LastBrushStamp.WorldNormal, LastBrushStamp.Falloff);
 }
 
 void UBaseBrushTool::ShutdownBrushStampIndicator()
diff --git a/Engine/Source/Runtime/Experimental/InteractiveToolsFramework/Public/BaseBehaviors/SingleClickBehavior.h b/Engine/Source/Runtime/Experimental/InteractiveToolsFramework/Public/BaseBehaviors/SingleClickBehavior.h
--- a/Engine/Source/Runtime/Experimental/InteractiveToolsFramework/Public/BaseBehaviors/SingleClickBehavior.h
+++ b/Engine/Source/Runtime/Experimental/InteractiveToolsFramework/Public/BaseBehaviors/SingleClickBehavior.h
@@ -5,6 +5,7 @@
 #include "CoreMinimal.h"
 #include "BaseBehaviors/AnyButtonInputBehavior.h"
 #include "BehaviorTargetInterfaces.h"
+#include "InputBehaviorModifierStates.h"
 #include "SingleClickBehavior.generated.h"
 
 
@@ -54,6 +55,11 @@ public:
 	UPROPERTY()
 	bool HitTestOnRelease;
 
+	/**
+	 * The modifier set for this behavior
+	 */
+	FInputBehaviorModifierStates Modifiers;
+
 
 protected:
 	/** Click Target object */
diff --git a/Engine/Source/Runtime/Experimental/InteractiveToolsFramework/Public/BaseGizmos/BrushStampIndicator.h b/Engine/Source/Runtime/Experimental/InteractiveToolsFramework/Public/BaseGizmos/BrushStampIndicator.h
--- a/Engine/Source/Runtime/Experimental/InteractiveToolsFramework/Public/BaseGizmos/BrushStampIndicator.h
+++ b/Engine/Source/Runtime/Experimental/InteractiveToolsFramework/Public/BaseGizmos/BrushStampIndicator.h
@@ -40,13 +40,16 @@ public:
 	/**
 	 * Update the Radius, Position, and Normal of the stamp indicator
 	 */
-	virtual void Update(float Radius, const FVector& Position, const FVector& Normal);
+	virtual void Update(float Radius, const FVector& Position, const FVector& Normal, float Falloff);
 
 public:
 
 	UPROPERTY()
 	float BrushRadius = 1.0f;
 
+	UPROPERTY()
+	float BrushFalloff = 0.5f;
+
 	UPROPERTY()
 	FVector BrushPosition = FVector::ZeroVector;
 
diff --git a/Engine/Source/Runtime/Experimental/InteractiveToolsFramework/Public/BaseTools/BaseBrushTool.h b/Engine/Source/Runtime/Experimental/InteractiveToolsFramework/Public/BaseTools/BaseBrushTool.h
--- a/Engine/Source/Runtime/Experimental/InteractiveToolsFramework/Public/BaseTools/BaseBrushTool.h
+++ b/Engine/Source/Runtime/Experimental/InteractiveToolsFramework/Public/BaseTools/BaseBrushTool.h
@@ -62,6 +62,9 @@ struct INTERACTIVETOOLSFRAMEWORK_API FBrushStampData
 
 	/** Hit Result provided by implementations - may not be fully populated */
 	FHitResult HitResult;
+
+	/** Falloff of brush stamp */
+	float Falloff;
 };
 
 
@@ -119,6 +122,10 @@ public:
 
 	virtual void IncreaseBrushSizeAction();
 	virtual void DecreaseBrushSizeAction();
+	virtual void IncreaseBrushStrengthAction();
+	virtual void DecreaseBrushStrengthAction();
+	virtual void IncreaseBrushFalloffAction();
+	virtual void DecreaseBrushFalloffAction();
 
 	virtual bool IsInBrushStroke() const { return bInBrushStroke; }
 
@@ -134,7 +141,8 @@ protected:
 	double CurrentBrushRadius;
 	void RecalculateBrushRadius();
 
-
+	UPROPERTY()
+	TSoftClassPtr<UBrushBaseProperties> PropertyClass;
 
 	//
 	// Brush Indicator support
