Added vertex preview size to the vertex paint brush. Vertex brush now respects if flow is enabled.

diff --git a/Engine/Plugins/MeshPainting/Source/MeshPaintEditorMode/Private/MeshPaintModeSettings.h b/Engine/Plugins/MeshPainting/Source/MeshPaintEditorMode/Private/MeshPaintModeSettings.h
--- a/Engine/Plugins/MeshPainting/Source/MeshPaintEditorMode/Private/MeshPaintModeSettings.h
+++ b/Engine/Plugins/MeshPainting/Source/MeshPaintEditorMode/Private/MeshPaintModeSettings.h
@@ -42,14 +42,10 @@ public:
 
 	UMeshPaintModeSettings()
 	{
-		VertexPreviewSize = 6.0f;
 		ColorViewMode = EMeshPaintDataColorViewMode::Normal;
 	}
 
 public:
-	/** Size of vertex points drawn when mesh painting is active. */
-	UPROPERTY(config, EditAnywhere, Category = Visualization)
-	float VertexPreviewSize;
 
 	/** Color view mode used to display Vertex Colors */
 	UPROPERTY(config, EditAnywhere, Category = Visualization)
diff --git a/Engine/Plugins/MeshPainting/Source/MeshPaintingToolset/Private/MeshPaintColorBrush.cpp b/Engine/Plugins/MeshPainting/Source/MeshPaintingToolset/Private/MeshPaintColorBrush.cpp
--- a/Engine/Plugins/MeshPainting/Source/MeshPaintingToolset/Private/MeshPaintColorBrush.cpp
+++ b/Engine/Plugins/MeshPainting/Source/MeshPaintingToolset/Private/MeshPaintColorBrush.cpp
@@ -11,7 +11,6 @@
 #include "ComponentReregisterContext.h"
 #include "MeshPaintAdapterFactory.h"
 #include "ToolDataVisualizer.h"
-//#include "MeshPaintModeSettings.h"
 
 
 #define LOCTEXT_NAMESPACE "MeshVertexBrush"
@@ -58,7 +57,9 @@ UInteractiveTool* UMeshWeightPaintingToolBuilder::BuildTool(const FToolBuilderSt
 UMeshVertexPaintingToolProperties::UMeshVertexPaintingToolProperties()
 	:UBrushBaseProperties(),
 	PaintColor(FLinearColor::White),
-	EraseColor(FLinearColor::Black)
+	EraseColor(FLinearColor::Black),
+	bEnableFlow(false),
+	VertexPreviewSize(6.0f)
 {
 }
 
@@ -70,6 +71,7 @@ void UMeshVertexPaintingToolProperties::SaveProperties(UInteractiveTool* SaveFro
 	UMeshVertexPaintingToolProperties* PropertyCache = GetPropertyCache<UMeshVertexPaintingToolProperties>();
 	PropertyCache->PaintColor = this->PaintColor;
 	PropertyCache->EraseColor = this->EraseColor;
+	PropertyCache->bEnableFlow = this->bEnableFlow;
 	PropertyCache->bOnlyFrontFacingTriangles = this->bOnlyFrontFacingTriangles;
 }
 
@@ -79,6 +81,7 @@ void UMeshVertexPaintingToolProperties::RestoreProperties(UInteractiveTool* Rest
 	UMeshVertexPaintingToolProperties* PropertyCache = GetPropertyCache<UMeshVertexPaintingToolProperties>();
 	this->PaintColor = PropertyCache->PaintColor;
 	this->EraseColor = PropertyCache->EraseColor;
+	this->bEnableFlow = PropertyCache->bEnableFlow;
 	this->bOnlyFrontFacingTriangles = PropertyCache->bOnlyFrontFacingTriangles;
 }
 
@@ -120,7 +123,7 @@ void UMeshVertexPaintingTool::Render(IToolsContextRenderAPI* RenderAPI)
 		static const FLinearColor NormalLineColor(0.3f, 1.0f, 0.3f);
 		const FLinearColor BrushCueColor = (bArePainting ? FLinearColor(1.0f, 1.0f, 0.3f) : FLinearColor(0.3f, 1.0f, 0.3f));
  		const FLinearColor InnerBrushCueColor = (bArePainting ? FLinearColor(0.5f, 0.5f, 0.1f) : FLinearColor(0.1f, 0.5f, 0.1f));
-		const float PointDrawSize = 6.0f;//GetDefault<UMeshPaintModeSettings>()->VertexPreviewSize;
+		const float PointDrawSize = VertexProperties->VertexPreviewSize;
 		// Draw trace surface normal
 		const FVector NormalLineEnd(LastBestHitResult.Location + LastBestHitResult.Normal * NormalLineSize);
 		Draw.DrawLine(FVector(LastBestHitResult.Location), NormalLineEnd, NormalLineColor, WidgetLineThickness);
@@ -180,7 +183,7 @@ void UMeshVertexPaintingTool::Tick(float DeltaTime)
 		bStampPending = false;
 
 		// flow
-		if (bInDrag)
+		if (bInDrag && VertexProperties && VertexProperties->bEnableFlow)
 		{
 			bStampPending = true;
 		}
diff --git a/Engine/Plugins/MeshPainting/Source/MeshPaintingToolset/Public/MeshPaintColorBrush.h b/Engine/Plugins/MeshPainting/Source/MeshPaintingToolset/Public/MeshPaintColorBrush.h
--- a/Engine/Plugins/MeshPainting/Source/MeshPaintingToolset/Public/MeshPaintColorBrush.h
+++ b/Engine/Plugins/MeshPainting/Source/MeshPaintingToolset/Public/MeshPaintColorBrush.h
@@ -92,10 +92,17 @@ public:
 	UPROPERTY(EditAnywhere, Category = VertexPainting)
 	FLinearColor EraseColor;
 
+	/** Enables "Flow" painting where paint is continually applied from the brush every tick */
+	UPROPERTY(EditAnywhere, Category = Brush, meta = (DisplayName = "Enable Brush Flow"))
+	bool bEnableFlow;
+
 	/** Whether back-facing triangles should be ignored */
 	UPROPERTY(EditAnywhere, Category = Brush, meta = (DisplayName = "Ignore Back-Facing"))
 	bool bOnlyFrontFacingTriangles;
 
+	/** Size of vertex points drawn when mesh painting is active. */
+	UPROPERTY(EditAnywhere, Category = Visualization)
+	float VertexPreviewSize;
 
 	virtual void SaveProperties(UInteractiveTool* SaveFromTool) override;
 	virtual void RestoreProperties(UInteractiveTool* RestoreToTool) override;
