Fix for the level editor viewport being in an incorrect color view after mode exit/editor restart

diff --git a/Engine/Plugins/MeshPainting/Source/MeshPaintEditorMode/Private/MeshPaintMode.cpp b/Engine/Plugins/MeshPainting/Source/MeshPaintEditorMode/Private/MeshPaintMode.cpp
--- a/Engine/Plugins/MeshPainting/Source/MeshPaintEditorMode/Private/MeshPaintMode.cpp
+++ b/Engine/Plugins/MeshPainting/Source/MeshPaintEditorMode/Private/MeshPaintMode.cpp
@@ -98,6 +98,7 @@ UMeshPaintMode::UMeshPaintMode()
 void UMeshPaintMode::Enter()
 {
 	Super::Enter();
+	GEditor->OnEditorClose().AddUObject(this, &UMeshPaintMode::OnResetViewMode);
 	ModeSettings = Cast<UMeshPaintModeSettings>(SettingsObject);
 	FMeshPaintEditorModeCommands ToolManagerCommands = FMeshPaintEditorModeCommands::Get();
 	RegisterTool(ToolManagerCommands.Select, TEXT("MeshClickTool"), NewObject<UMeshClickToolBuilder>());
@@ -111,6 +112,8 @@ void UMeshPaintMode::Enter()
 
 void UMeshPaintMode::Exit()
 {
+	GEditor->OnEditorClose().RemoveAll(this);
+	OnResetViewMode();
 	const FMeshPaintEditorModeCommands& Commands = FMeshPaintEditorModeCommands::Get();
 	const TSharedRef<FUICommandList>& CommandList = Toolkit->GetToolkitCommands();
 	for (const TSharedPtr<const FUICommandInfo> Action : Commands.Commands[UMeshPaintMode::MeshPaintMode_Color])
@@ -682,5 +685,20 @@ void UMeshPaintMode::UpdateOnPaletteChange()
 	UpdateSelectedMeshes();
 }
 
+void UMeshPaintMode::OnResetViewMode()
+{
+	// Reset viewport color mode for all active viewports
+	for (FEditorViewportClient* ViewportClient : GEditor->GetAllViewportClients())
+	{
+		if (!ViewportClient || ViewportClient->GetModeTools() != GetModeManager())
+		{
+			continue;
+		}
+
+		UMeshPaintModeHelpers::SetViewportColorMode(EMeshPaintDataColorViewMode::Normal, ViewportClient);
+	}
+}
+
+
 #undef LOCTEXT_NAMESPACE
 
diff --git a/Engine/Plugins/MeshPainting/Source/MeshPaintEditorMode/Private/MeshPaintMode.h b/Engine/Plugins/MeshPainting/Source/MeshPaintEditorMode/Private/MeshPaintMode.h
--- a/Engine/Plugins/MeshPainting/Source/MeshPaintEditorMode/Private/MeshPaintMode.h
+++ b/Engine/Plugins/MeshPainting/Source/MeshPaintEditorMode/Private/MeshPaintMode.h
@@ -67,9 +67,9 @@ protected:
 	virtual void ActorSelectionChangeNotify() override;
 	virtual void ActivateDefaultTool() override;
 	virtual void UpdateOnPaletteChange() override;
-
 	// end UEdMode Interface
 	void UpdateSelectedMeshes();
+	void OnResetViewMode();
 	void FillWithVertexColor();
 	void PropagateVertexColorsToAsset();
 	bool CanPropagateVertexColors() const;
