Fix for CIS error

diff --git a/Engine/Plugins/MeshPainting/Source/MeshPaintingToolset/Private/MeshPaintHelpers.cpp b/Engine/Plugins/MeshPainting/Source/MeshPaintingToolset/Private/MeshPaintHelpers.cpp
--- a/Engine/Plugins/MeshPainting/Source/MeshPaintingToolset/Private/MeshPaintHelpers.cpp
+++ b/Engine/Plugins/MeshPainting/Source/MeshPaintingToolset/Private/MeshPaintHelpers.cpp
@@ -166,7 +166,7 @@ bool UMeshPaintingToolset::PaintVertex(const FVector& InVertexPosition, const FM
 	if (UMeshPaintingToolset::IsPointInfluencedByBrush(InVertexPosition, InParams, SquaredDistanceToVertex2D, VertexDepthToBrush))
 	{
 		// Compute amount of paint to apply
-		const float PaintAmount = ComputePaintMultiplier(SquaredDistanceToVertex2D, InParams.BrushStrength, InParams.InnerBrushRadius, InParams.BrushRadialFalloffRange, InParams.BrushDepth, InParams.BrushDepthFalloffRange, VertexDepthToBrush);
+		const float PaintAmount = UMeshPaintingToolset::ComputePaintMultiplier(SquaredDistanceToVertex2D, InParams.BrushStrength, InParams.InnerBrushRadius, InParams.BrushRadialFalloffRange, InParams.BrushDepth, InParams.BrushDepthFalloffRange, VertexDepthToBrush);
 			
 		const FLinearColor OldColor = InOutVertexColor.ReinterpretAsLinear();
 		FLinearColor NewColor = OldColor;
diff --git a/Engine/Plugins/MeshPainting/Source/MeshPaintingToolset/Public/MeshPaintHelpers.h b/Engine/Plugins/MeshPainting/Source/MeshPaintingToolset/Public/MeshPaintHelpers.h
--- a/Engine/Plugins/MeshPainting/Source/MeshPaintingToolset/Public/MeshPaintHelpers.h
+++ b/Engine/Plugins/MeshPainting/Source/MeshPaintingToolset/Public/MeshPaintHelpers.h
@@ -242,7 +242,7 @@ void UMeshPaintingToolset::ApplyBrushToVertex(const FVector& VertexPosition, con
 	if (IsPointInfluencedByBrush(BrushSpacePosition2D, BrushRadius * BrushRadius, InfluencedValue))
 	{
 		float InnerBrushRadius = BrushFalloffAmount * BrushRadius;
-		float PaintStrength = MeshPaintHelpers::ComputePaintMultiplier(BrushSpacePosition2D.SizeSquared(), BrushStrength, InnerBrushRadius, BrushRadius - InnerBrushRadius, 1.0f, 1.0f, 1.0f);
+		float PaintStrength = UMeshPaintingToolset::ComputePaintMultiplier(BrushSpacePosition2D.SizeSquared(), BrushStrength, InnerBrushRadius, BrushRadius - InnerBrushRadius, 1.0f, 1.0f, 1.0f);
 
 		const T OldValue = InOutValue;
 		InOutValue = FMath::LerpStable(OldValue, PaintValue, PaintStrength);
