Setting up details customization for vertex weight painting

diff --git a/Engine/Plugins/MeshPainting/Source/MeshPaintEditorMode/Private/BrushSettingsCustomization.cpp b/Engine/Plugins/MeshPainting/Source/MeshPaintEditorMode/Private/BrushSettingsCustomization.cpp
--- a/Engine/Plugins/MeshPainting/Source/MeshPaintEditorMode/Private/BrushSettingsCustomization.cpp
+++ b/Engine/Plugins/MeshPainting/Source/MeshPaintEditorMode/Private/BrushSettingsCustomization.cpp
@@ -326,6 +326,49 @@ TSharedRef<IDetailCustomization> FWeightPaintingSettingsCustomization::MakeInsta
 void FWeightPaintingSettingsCustomization::CustomizeDetails(IDetailLayoutBuilder& DetailLayout)
 {
 	FVertexPaintingSettingsCustomization::CustomizeDetails(DetailLayout);
+
+	TMap<FName, TSharedRef<IPropertyHandle>> CustomizedProperties;
+	TMap<FName, TSharedRef<IPropertyHandle>> Properties;
+
+	TSharedRef<IPropertyHandle> WeightTypeProperty = DetailLayout.GetProperty(GET_MEMBER_NAME_CHECKED(UMeshWeightPaintingToolProperties, TextureWeightType));
+	TSharedRef<IPropertyHandle> PaintWeightProperty = DetailLayout.GetProperty(GET_MEMBER_NAME_CHECKED(UMeshWeightPaintingToolProperties, PaintTextureWeightIndex));
+	TSharedRef<IPropertyHandle> EraseWeightProperty = DetailLayout.GetProperty(GET_MEMBER_NAME_CHECKED(UMeshWeightPaintingToolProperties, EraseTextureWeightIndex));
+
+	WeightTypeProperty->SetOnPropertyValueChanged(FSimpleDelegate::CreateRaw(this, &FWeightPaintingSettingsCustomization::OnTextureWeightTypeChanged, WeightTypeProperty, PaintWeightProperty, EraseWeightProperty));
+
+	static FText RestrictReason = NSLOCTEXT("VertexPaintSettings", "TextureIndexRestriction", "Unable to paint this Texture, change Texture Weight Type");
+	BlendPaintEnumRestriction = MakeShareable(new FPropertyRestriction(RestrictReason));
+	PaintWeightProperty->AddRestriction(BlendPaintEnumRestriction.ToSharedRef());
+	EraseWeightProperty->AddRestriction(BlendPaintEnumRestriction.ToSharedRef());
+
+	OnTextureWeightTypeChanged(WeightTypeProperty, PaintWeightProperty, EraseWeightProperty);
+}
+
+void FWeightPaintingSettingsCustomization::OnTextureWeightTypeChanged(TSharedRef<IPropertyHandle> WeightTypeProperty, TSharedRef<IPropertyHandle> PaintWeightProperty, TSharedRef<IPropertyHandle> EraseWeightProperty)
+{
+	UEnum* ImportTypeEnum = StaticEnum<EMeshPaintTextureIndex>();
+	uint8 EnumValue = 0;
+	WeightTypeProperty->GetValue(EnumValue);
+
+	BlendPaintEnumRestriction->RemoveAll();
+	for (uint8 EnumIndex = 0; EnumIndex < (ImportTypeEnum->GetMaxEnumValue() + 1); ++EnumIndex)
+	{
+		if ((EnumIndex + 1) > EnumValue)
+		{
+			FString EnumName = ImportTypeEnum->GetNameByValue(EnumIndex).ToString();
+			EnumName.RemoveFromStart("EMeshPaintTextureIndex::");
+			BlendPaintEnumRestriction->AddDisabledValue(EnumName);
+		}
+	}
+
+	uint8 Value = 0;
+	PaintWeightProperty->GetValue(Value);
+	Value = FMath::Clamp<uint8>(Value, 0, EnumValue - 1);
+	PaintWeightProperty->SetValue(Value);
+
+	EraseWeightProperty->GetValue(Value);
+	Value = FMath::Clamp<uint8>(Value, 0, EnumValue - 1);
+	EraseWeightProperty->SetValue(Value);
 }
 
 #undef LOCTEXT_NAMESPACE
\ No newline at end of file
diff --git a/Engine/Plugins/MeshPainting/Source/MeshPaintEditorMode/Private/BrushSettingsCustomization.h b/Engine/Plugins/MeshPainting/Source/MeshPaintEditorMode/Private/BrushSettingsCustomization.h
--- a/Engine/Plugins/MeshPainting/Source/MeshPaintEditorMode/Private/BrushSettingsCustomization.h
+++ b/Engine/Plugins/MeshPainting/Source/MeshPaintEditorMode/Private/BrushSettingsCustomization.h
@@ -44,6 +44,12 @@ public:
 
 	/** IPropertyTypeCustomization interface */
 	virtual void CustomizeDetails(IDetailLayoutBuilder& DetailLayout) override;
+protected:
+	void OnTextureWeightTypeChanged(TSharedRef<IPropertyHandle> WeightTypeProperty, TSharedRef<IPropertyHandle> PaintWeightProperty, TSharedRef<IPropertyHandle> EraseWeightProperty);
+
+	/** Property restriction applied to blend paint enum dropdown box */
+	TSharedPtr<FPropertyRestriction> BlendPaintEnumRestriction;
+
 };
 
 // class FTexturePaintSettingsCustomization : public IPropertyTypeCustomization
diff --git a/Engine/Plugins/MeshPainting/Source/MeshPaintingToolset/Public/MeshPaintColorBrush.h b/Engine/Plugins/MeshPainting/Source/MeshPaintingToolset/Public/MeshPaintColorBrush.h
--- a/Engine/Plugins/MeshPainting/Source/MeshPaintingToolset/Public/MeshPaintColorBrush.h
+++ b/Engine/Plugins/MeshPainting/Source/MeshPaintingToolset/Public/MeshPaintColorBrush.h
@@ -101,7 +101,7 @@ public:
 	bool bOnlyFrontFacingTriangles;
 
 	/** Size of vertex points drawn when mesh painting is active. */
-	UPROPERTY(EditAnywhere, Category = Visualization)
+	UPROPERTY(EditAnywhere, Category = "VertexPainting|Visualization")
 	float VertexPreviewSize;
 
 	virtual void SaveProperties(UInteractiveTool* SaveFromTool) override;
