Renaming brush tool files

diff --git a/Engine/Plugins/MeshPainting/Source/MeshPaintEditorMode/Private/BrushSettingsCustomization.cpp b/Engine/Plugins/MeshPainting/Source/MeshPaintEditorMode/Private/BrushSettingsCustomization.cpp
--- a/Engine/Plugins/MeshPainting/Source/MeshPaintEditorMode/Private/BrushSettingsCustomization.cpp
+++ b/Engine/Plugins/MeshPainting/Source/MeshPaintEditorMode/Private/BrushSettingsCustomization.cpp
@@ -18,7 +18,7 @@
 #include "Widgets/Layout/SBox.h"
 #include "ScopedTransaction.h"
 #include "MeshPaintHelpers.h"
-#include "MeshPaintColorBrush.h"
+#include "MeshVertexPaintingTool.h"
 #include "MeshPaintHelpers.h"
 #include "InteractiveToolManager.h"
 #include "Editor.h"
diff --git a/Engine/Plugins/MeshPainting/Source/MeshPaintEditorMode/Private/MeshPaintMode.h b/Engine/Plugins/MeshPainting/Source/MeshPaintEditorMode/Private/MeshPaintMode.h
--- a/Engine/Plugins/MeshPainting/Source/MeshPaintEditorMode/Private/MeshPaintMode.h
+++ b/Engine/Plugins/MeshPainting/Source/MeshPaintEditorMode/Private/MeshPaintMode.h
@@ -13,7 +13,7 @@
 #include "Framework/Commands/UICommandList.h"
 #include "IMeshPaintGeometryAdapter.h"
 #include "MeshPaintHelpers.h"
-#include "MeshPaintColorBrush.h"
+#include "MeshVertexPaintingTool.h"
 #include "MeshPaintMode.generated.h"
 
 class UMeshVertexPaintingToolProperties;
diff --git a/Engine/Plugins/MeshPainting/Source/MeshPaintEditorMode/Private/MeshPaintModeCommands.cpp b/Engine/Plugins/MeshPainting/Source/MeshPaintEditorMode/Private/MeshPaintModeCommands.cpp
--- a/Engine/Plugins/MeshPainting/Source/MeshPaintEditorMode/Private/MeshPaintModeCommands.cpp
+++ b/Engine/Plugins/MeshPainting/Source/MeshPaintEditorMode/Private/MeshPaintModeCommands.cpp
@@ -4,7 +4,7 @@
 #include "Framework/Commands/UIAction.h"
 #include "MeshPaintMode.h"
 #include "SingleSelectionTool.h"
-#include "MeshPaintColorBrush.h"
+#include "MeshVertexPaintingTool.h"
 
 #define LOCTEXT_NAMESPACE "MeshPaintEditorModeCommands"
 
diff --git a/Engine/Plugins/MeshPainting/Source/MeshPaintEditorMode/Private/MeshPaintModeHelpers.cpp b/Engine/Plugins/MeshPainting/Source/MeshPaintEditorMode/Private/MeshPaintModeHelpers.cpp
--- a/Engine/Plugins/MeshPainting/Source/MeshPaintEditorMode/Private/MeshPaintModeHelpers.cpp
+++ b/Engine/Plugins/MeshPainting/Source/MeshPaintEditorMode/Private/MeshPaintModeHelpers.cpp
@@ -44,7 +44,7 @@
 #include "Rendering/SkeletalMeshModel.h"
 #include "MeshPaintHelpers.h"
 #include "MeshPaintMode.h"
-#include "MeshPaintColorBrush.h"
+#include "MeshVertexPaintingTool.h"
 
 
 void UMeshPaintModeHelpers::SetViewportColorMode(EMeshPaintDataColorViewMode ColorViewMode, FEditorViewportClient* ViewportClient)
diff --git a/Engine/Plugins/MeshPainting/Source/MeshPaintingToolset/Private/MeshPaintColorBrush.cpp b/Engine/Plugins/MeshPainting/Source/MeshPaintingToolset/Private/MeshPaintColorBrush.cpp
deleted file mode 100644
diff --git a/Engine/Plugins/MeshPainting/Source/MeshPaintingToolset/Private/MeshPaintHelpers.cpp b/Engine/Plugins/MeshPainting/Source/MeshPaintingToolset/Private/MeshPaintHelpers.cpp
--- a/Engine/Plugins/MeshPainting/Source/MeshPaintingToolset/Private/MeshPaintHelpers.cpp
+++ b/Engine/Plugins/MeshPainting/Source/MeshPaintingToolset/Private/MeshPaintHelpers.cpp
@@ -18,7 +18,7 @@
 #include "Async/ParallelFor.h"
 #include "Rendering/SkeletalMeshModel.h"
 #include "InteractiveToolManager.h"
-#include "MeshPaintColorBrush.h"
+#include "MeshVertexPaintingTool.h"
 #include "InputRouter.h"
 #include "InteractiveGizmoManager.h"
 #include "ToolContextInterfaces.h"
diff --git a/Engine/Plugins/MeshPainting/Source/MeshPaintingToolset/Private/MeshVertexPaintingTool.cpp b/Engine/Plugins/MeshPainting/Source/MeshPaintingToolset/Private/MeshVertexPaintingTool.cpp
new file mode 100644
--- /dev/null
+++ b/Engine/Plugins/MeshPainting/Source/MeshPaintingToolset/Private/MeshVertexPaintingTool.cpp
@@ -0,0 +1,710 @@
+// Copyright 1998-2019 Epic Games, Inc. All Rights Reserved.
+
+#include "MeshVertexPaintingTool.h"
+#include "InteractiveToolManager.h"
+#include "ToolBuilderUtil.h"
+#include "SceneManagement.h" // for FPrimitiveDrawInterface
+#include "ToolContextInterfaces.h"
+#include "Components/MeshComponent.h"
+#include "Math/UnrealMathUtility.h"
+#include "IMeshPaintGeometryAdapter.h"
+#include "ComponentReregisterContext.h"
+#include "MeshPaintAdapterFactory.h"
+#include "ToolDataVisualizer.h"
+
+
+#define LOCTEXT_NAMESPACE "MeshVertexBrush"
+
+/*
+ * ToolBuilder
+ */
+
+bool HasPaintableMesh(UActorComponent* Component)
+{
+	return Cast<UMeshComponent>(Component) != nullptr;
+}
+
+
+bool UMeshColorPaintingToolBuilder::CanBuildTool(const FToolBuilderState& SceneState) const
+{
+	UActorComponent* ActorComponent = ToolBuilderUtil::FindFirstComponent(SceneState, HasPaintableMesh);
+	return ActorComponent != nullptr;
+}
+
+UInteractiveTool* UMeshColorPaintingToolBuilder::BuildTool(const FToolBuilderState& SceneState) const
+{
+	UMeshColorPaintingTool* NewTool = NewObject<UMeshColorPaintingTool>(SceneState.ToolManager, "ColorBrushTool");
+	return NewTool;
+}
+
+bool UMeshWeightPaintingToolBuilder::CanBuildTool(const FToolBuilderState& SceneState) const
+{
+	UActorComponent* ActorComponent = ToolBuilderUtil::FindFirstComponent(SceneState, HasPaintableMesh);
+	return ActorComponent != nullptr;
+}
+
+UInteractiveTool* UMeshWeightPaintingToolBuilder::BuildTool(const FToolBuilderState& SceneState) const
+{
+	UMeshWeightPaintingTool* NewTool = NewObject<UMeshWeightPaintingTool>(SceneState.ToolManager, "ColorWeightTool");
+	return NewTool;
+}
+
+
+/*
+ * Tool
+ */
+
+UMeshVertexPaintingToolProperties::UMeshVertexPaintingToolProperties()
+	:UBrushBaseProperties(),
+	PaintColor(FLinearColor::White),
+	EraseColor(FLinearColor::Black),
+	bEnableFlow(false),
+	VertexPreviewSize(6.0f)
+{
+}
+
+
+
+void UMeshVertexPaintingToolProperties::SaveProperties(UInteractiveTool* SaveFromTool)
+{
+	UBrushBaseProperties::SaveProperties(SaveFromTool);
+	UMeshVertexPaintingToolProperties* PropertyCache = GetPropertyCache<UMeshVertexPaintingToolProperties>();
+	PropertyCache->PaintColor = this->PaintColor;
+	PropertyCache->EraseColor = this->EraseColor;
+	PropertyCache->bEnableFlow = this->bEnableFlow;
+	PropertyCache->bOnlyFrontFacingTriangles = this->bOnlyFrontFacingTriangles;
+}
+
+void UMeshVertexPaintingToolProperties::RestoreProperties(UInteractiveTool* RestoreToTool)
+{
+	UBrushBaseProperties::RestoreProperties(RestoreToTool);
+	UMeshVertexPaintingToolProperties* PropertyCache = GetPropertyCache<UMeshVertexPaintingToolProperties>();
+	this->PaintColor = PropertyCache->PaintColor;
+	this->EraseColor = PropertyCache->EraseColor;
+	this->bEnableFlow = PropertyCache->bEnableFlow;
+	this->bOnlyFrontFacingTriangles = PropertyCache->bOnlyFrontFacingTriangles;
+}
+
+UMeshVertexPaintingTool::UMeshVertexPaintingTool()
+{
+	PropertyClass = UMeshVertexPaintingToolProperties::StaticClass();
+}
+
+
+void UMeshVertexPaintingTool::Setup()
+{
+	Super::Setup();
+	VertexProperties = Cast<UMeshVertexPaintingToolProperties>(BrushProperties);
+	bResultValid = false;
+	bStampPending = false;
+	BrushProperties->RestoreProperties(this);
+}
+
+
+void UMeshVertexPaintingTool::Shutdown(EToolShutdownType ShutdownType)
+{
+	FinishPainting();
+	BrushProperties->SaveProperties(this);
+	Super::Shutdown(ShutdownType);
+}
+
+void UMeshVertexPaintingTool::Render(IToolsContextRenderAPI* RenderAPI)
+{
+	UMeshToolManager* MeshToolManager = Cast<UMeshToolManager>(GetToolManager());
+	if (MeshToolManager && LastBestHitResult.Component != nullptr)
+	{
+		Super::Render(RenderAPI);
+		FToolDataVisualizer Draw;
+		Draw.BeginFrame(RenderAPI);
+		static float WidgetLineThickness = 1.0f;
+		static FLinearColor VertexPointColor = FLinearColor::White;
+		static FLinearColor	HoverVertexPointColor = FLinearColor(0.3f, 1.0f, 0.3f);
+		const float NormalLineSize(BrushProperties->BrushRadius * 0.35f);	// Make the normal line length a function of brush size
+		static const FLinearColor NormalLineColor(0.3f, 1.0f, 0.3f);
+		const FLinearColor BrushCueColor = (bArePainting ? FLinearColor(1.0f, 1.0f, 0.3f) : FLinearColor(0.3f, 1.0f, 0.3f));
+ 		const FLinearColor InnerBrushCueColor = (bArePainting ? FLinearColor(0.5f, 0.5f, 0.1f) : FLinearColor(0.1f, 0.5f, 0.1f));
+		const float PointDrawSize = VertexProperties->VertexPreviewSize;
+		// Draw trace surface normal
+		const FVector NormalLineEnd(LastBestHitResult.Location + LastBestHitResult.Normal * NormalLineSize);
+		Draw.DrawLine(FVector(LastBestHitResult.Location), NormalLineEnd, NormalLineColor, WidgetLineThickness);
+
+		for (UMeshComponent* CurrentComponent : MeshToolManager->GetPaintableMeshComponents())
+		{
+			TSharedPtr<IMeshPaintGeometryAdapter> MeshAdapter = MeshToolManager->GetAdapterForComponent(Cast<UMeshComponent>(CurrentComponent));
+
+			if (MeshAdapter->IsValid() /*&& bRenderVertices*/ && MeshAdapter->SupportsVertexPaint())
+			{
+				const FMatrix ComponentToWorldMatrix = MeshAdapter->GetComponentToWorldMatrix();
+				FViewCameraState CameraState;
+				GetToolManager()->GetContextQueriesAPI()->GetCurrentViewState(CameraState);
+				const FVector ComponentSpaceCameraPosition(ComponentToWorldMatrix.InverseTransformPosition(CameraState.Position));
+				const FVector ComponentSpaceBrushPosition(ComponentToWorldMatrix.InverseTransformPosition(LastBestHitResult.Location));
+
+				// @todo MeshPaint: Input vector doesn't work well with non-uniform scale
+				const float ComponentSpaceBrushRadius = ComponentToWorldMatrix.InverseTransformVector(FVector(BrushProperties->BrushRadius, 0.0f, 0.0f)).Size();
+				const float ComponentSpaceSquaredBrushRadius = ComponentSpaceBrushRadius * ComponentSpaceBrushRadius;
+
+				const TArray<FVector>& InRangeVertices = MeshAdapter->SphereIntersectVertices(ComponentSpaceSquaredBrushRadius, ComponentSpaceBrushPosition, ComponentSpaceCameraPosition, VertexProperties->bOnlyFrontFacingTriangles);
+
+				for (const FVector& Vertex : InRangeVertices)
+				{
+					const FVector WorldPositionVertex = ComponentToWorldMatrix.TransformPosition(Vertex);
+					if ((LastBestHitResult.Location - WorldPositionVertex).Size() <= BrushProperties->BrushRadius)
+					{
+						const float VisualBiasDistance = 0.15f;
+						const FVector VertexVisualPosition = WorldPositionVertex + LastBestHitResult.Normal * VisualBiasDistance;
+						Draw.DrawPoint(VertexVisualPosition, HoverVertexPointColor, PointDrawSize, SDPG_World);
+					}
+				}
+			}
+		}
+		Draw.EndFrame();
+	}
+	UpdateResult();
+
+}
+
+void UMeshVertexPaintingTool::Tick(float DeltaTime)
+{
+	Super::Tick(DeltaTime);
+
+	if (UMeshToolManager* MeshToolManager = Cast<UMeshToolManager>(GetToolManager()))
+	{
+		if (MeshToolManager->bNeedsRecache)
+		{
+			CacheSelectionData();
+			bDoRestoreRenTargets = true;
+		}
+	}
+
+	if (bStampPending)
+	{
+		Paint(PendingStampRay.Origin, PendingStampRay.Direction);
+		bStampPending = false;
+
+		// flow
+		if (bInDrag && VertexProperties && VertexProperties->bEnableFlow)
+		{
+			bStampPending = true;
+		}
+	}
+}
+
+void UMeshVertexPaintingTool::OnPropertyModified(UObject* PropertySet, UProperty* Property)
+{
+	bResultValid = false;
+}
+
+
+double UMeshVertexPaintingTool::CalculateTargetEdgeLength(int TargetTriCount)
+{
+	double TargetTriArea = InitialMeshArea / (double)TargetTriCount;
+	double EdgeLen = (TargetTriArea);
+	return (double)FMath::RoundToInt(EdgeLen*100.0) / 100.0;
+}
+
+bool UMeshVertexPaintingTool::Paint(const FVector& InRayOrigin, const FVector& InRayDirection)
+{
+	// Determine paint action according to whether or not shift is held down
+	const EMeshPaintModeAction PaintAction = GetShiftToggle() ? EMeshPaintModeAction::Erase : EMeshPaintModeAction::Paint;
+	const float PaintStrength = 1.0f; //Viewport->IsPenActive() ? Viewport->GetTabletPressure() : 1.f;
+	// Handle internal painting functionality
+	TPair<FVector, FVector> Ray(InRayOrigin, InRayDirection);
+	return PaintInternal(MakeArrayView(&Ray, 1), PaintAction, PaintStrength);
+}
+
+bool UMeshVertexPaintingTool::Paint(const TArrayView<TPair<FVector, FVector>>& Rays)
+{
+	// Determine paint action according to whether or not shift is held down
+	const EMeshPaintModeAction PaintAction = GetShiftToggle() ? EMeshPaintModeAction::Erase : EMeshPaintModeAction::Paint;
+
+	const float PaintStrength = 1.0f; //Viewport->IsPenActive() ? Viewport->GetTabletPressure() : 1.f;
+	// Handle internal painting functionality
+	return PaintInternal(Rays, PaintAction, PaintStrength);
+}
+
+
+bool UMeshVertexPaintingTool::PaintInternal(const TArrayView<TPair<FVector, FVector>>& Rays, EMeshPaintModeAction PaintAction, float PaintStrength)
+{
+	TArray<FPaintRayResults> PaintRayResults;
+	PaintRayResults.AddDefaulted(Rays.Num());
+	LastBestHitResult.Reset();
+	TMap<UMeshComponent*, TArray<int32>> HoveredComponents;
+
+	const float BrushRadius = BrushProperties->BrushRadius;
+	const bool bIsPainting = PaintAction == EMeshPaintModeAction::Paint;
+	const float InStrengthScale = PaintStrength;
+
+	bool bPaintApplied = false;
+	UMeshToolManager* MeshToolManager = Cast<UMeshToolManager>(GetToolManager());
+	// Fire out a ray to see if there is a *selected* component under the mouse cursor that can be painted.
+	for (int32 i = 0; i < Rays.Num(); ++i)
+	{
+		const FVector& RayOrigin = Rays[i].Key;
+		const FVector& RayDirection = Rays[i].Value;
+		FRay Ray = FRay(RayOrigin, RayDirection);
+		FHitResult& BestTraceResult = PaintRayResults[i].BestTraceResult;
+		const FVector TraceStart(RayOrigin);
+		const FVector TraceEnd(RayOrigin + RayDirection * HALF_WORLD_MAX);
+
+		for (UMeshComponent* MeshComponent : MeshToolManager->GetPaintableMeshComponents())
+		{
+			TSharedPtr<IMeshPaintGeometryAdapter> MeshAdapter = MeshToolManager->GetAdapterForComponent(MeshComponent);
+
+			// Ray trace
+			FHitResult TraceHitResult(1.0f);
+
+			if (MeshAdapter->LineTraceComponent(TraceHitResult, TraceStart, TraceEnd, FCollisionQueryParams(SCENE_QUERY_STAT(Paint), true)))
+			{
+				// Find the closest impact
+				if ((BestTraceResult.GetComponent() == nullptr) || (TraceHitResult.Time < BestTraceResult.Time))
+				{
+					BestTraceResult = TraceHitResult;
+				}
+			}
+		}
+
+		bool bUsed = false;
+
+		if (BestTraceResult.GetComponent() != nullptr)
+		{
+			FBox BrushBounds = FBox::BuildAABB(BestTraceResult.Location, FVector(BrushRadius * 1.25f, BrushRadius * 1.25f, BrushRadius * 1.25f));
+
+			// Vertex paint mode, so we want all valid components overlapping the brush hit location
+			for (auto TestComponent : MeshToolManager->GetPaintableMeshComponents())
+			{
+				const FBox ComponentBounds = TestComponent->Bounds.GetBox();
+
+				if (MeshToolManager->GetComponentToAdapterMap().Contains(TestComponent) && ComponentBounds.Intersect(BrushBounds))
+				{
+					// OK, this mesh potentially overlaps the brush!
+					HoveredComponents.FindOrAdd(TestComponent).Add(i);
+					bUsed = true;
+				}
+			}
+		}
+		if (bUsed)
+		{
+			FVector BrushXAxis, BrushYAxis;
+			BestTraceResult.Normal.FindBestAxisVectors(BrushXAxis, BrushYAxis);
+			// Display settings
+			const float VisualBiasDistance = 0.15f;
+			const FVector BrushVisualPosition = BestTraceResult.Location + BestTraceResult.Normal * VisualBiasDistance;
+
+			const FLinearColor PaintColor = VertexProperties->PaintColor;
+			const FLinearColor EraseColor = VertexProperties->EraseColor;
+
+			// NOTE: We square the brush strength to maximize slider precision in the low range
+			const float BrushStrength = BrushProperties->BrushStrength *  BrushProperties->BrushStrength * InStrengthScale;
+
+			const float BrushDepth = BrushRadius;
+			LastBestHitResult = BestTraceResult;
+			// Mesh paint settings
+			FMeshPaintParameters& Params = PaintRayResults[i].Params;
+			Params.PaintAction = PaintAction;
+			Params.BrushPosition = BestTraceResult.Location;
+			Params.BrushNormal = BestTraceResult.Normal;
+			Params.BrushColor = bIsPainting ? PaintColor : EraseColor;
+			Params.SquaredBrushRadius = BrushRadius * BrushRadius;
+			Params.BrushRadialFalloffRange = BrushProperties->BrushFalloffAmount * BrushRadius;
+			Params.InnerBrushRadius = BrushRadius - Params.BrushRadialFalloffRange;
+			Params.BrushDepth = BrushDepth;
+			Params.BrushDepthFalloffRange = BrushProperties->BrushFalloffAmount * BrushDepth;
+			Params.BrushDepthFalloffRange = BrushProperties->BrushFalloffAmount * BrushDepth;
+			Params.InnerBrushDepth = BrushDepth - Params.BrushDepthFalloffRange;
+			Params.BrushStrength = BrushStrength;
+			Params.BrushToWorldMatrix = FMatrix(BrushXAxis, BrushYAxis, Params.BrushNormal, Params.BrushPosition);
+			Params.InverseBrushToWorldMatrix = Params.BrushToWorldMatrix.Inverse();
+
+			SetAdditionalPaintParameters(Params);
+		
+		}
+	}
+
+	if (HoveredComponents.Num() > 0)
+	{
+		if (bArePainting == false)
+		{
+				// Vertex painting is an ongoing transaction, while texture painting is handled separately later in a single transaction
+				GetToolManager()->BeginUndoTransaction(LOCTEXT("MeshPaintMode_VertexPaint_TransactionPaintStroke", "Vertex Paint"));
+				bArePainting = true;
+				TimeSinceStartedPainting = 0.0f;
+		}
+
+		// Iterate over the selected meshes under the cursor and paint them!
+		for (auto& Entry : HoveredComponents)
+		{
+			UMeshComponent* HoveredComponent = Entry.Key;
+			TArray<int32>& PaintRayResultIds = Entry.Value;
+
+			IMeshPaintGeometryAdapter* MeshAdapter = MeshToolManager->GetAdapterForComponent(HoveredComponent).Get();
+			if (!ensure(MeshAdapter))
+			{
+				continue;
+			}
+
+			if (MeshAdapter->SupportsVertexPaint())
+			{
+				FPerVertexPaintActionArgs Args;
+				Args.Adapter = MeshAdapter;
+				FViewCameraState CameraState;
+				GetToolManager()->GetContextQueriesAPI()->GetCurrentViewState(CameraState);
+				Args.CameraPosition = CameraState.Position;
+				Args.BrushProperties = VertexProperties;
+				Args.Action = PaintAction;
+
+				bool bMeshPreEditCalled = false;
+
+				TSet<int32> InfluencedVertices;
+				for (int32 PaintRayResultId : PaintRayResultIds)
+				{
+					InfluencedVertices.Reset();
+					Args.HitResult = PaintRayResults[PaintRayResultId].BestTraceResult;
+					bPaintApplied |= UMeshPaintingToolset::GetPerVertexPaintInfluencedVertices(Args, InfluencedVertices);
+
+					if (InfluencedVertices.Num() == 0)
+					{
+						continue;
+					}
+
+					if (!bMeshPreEditCalled)
+					{
+						bMeshPreEditCalled = true;
+						MeshAdapter->PreEdit();
+					}
+
+					for (const int32 VertexIndex : InfluencedVertices)
+					{
+						ApplyVertexData(Args, VertexIndex, PaintRayResults[PaintRayResultId].Params);
+					}
+				}
+
+				if (bMeshPreEditCalled)
+				{
+					MeshAdapter->PostEdit();
+				}
+			}
+		}
+	}
+
+	return bPaintApplied;
+}
+
+void UMeshVertexPaintingTool::ApplyVertexData(FPerVertexPaintActionArgs& InArgs, int32 VertexIndex, FMeshPaintParameters Parameters)
+{
+	/** Retrieve vertex position and color for applying vertex painting */
+	FColor PaintColor;
+	FVector Position;
+	InArgs.Adapter->GetVertexPosition(VertexIndex, Position);
+	Position = InArgs.Adapter->GetComponentToWorldMatrix().TransformPosition(Position);
+	InArgs.Adapter->GetVertexColor(VertexIndex, PaintColor, true);
+	UMeshPaintingToolset::PaintVertex(Position, Parameters, PaintColor);
+	InArgs.Adapter->SetVertexColor(VertexIndex, PaintColor, true);
+}
+
+
+void UMeshVertexPaintingTool::UpdateResult()
+{
+	GetToolManager()->PostInvalidation();
+
+	bResultValid = true;
+}
+
+bool UMeshVertexPaintingTool::HasAccept() const
+{
+	return false;
+}
+
+bool UMeshVertexPaintingTool::CanAccept() const
+{
+	return false;
+}
+
+void UMeshVertexPaintingTool::OnBeginDrag(const FRay& Ray)
+{
+	Super::OnBeginDrag(Ray);
+	FHitResult OutHit;
+	if (HitTest(Ray, OutHit))
+	{
+		bInDrag = true;
+
+		// apply initial stamp
+		PendingStampRay = Ray;
+		bStampPending = true;
+	}
+}
+
+void UMeshVertexPaintingTool::OnUpdateDrag(const FRay& Ray)
+{
+	Super::OnUpdateDrag(Ray);
+	if (bInDrag)
+	{
+		PendingStampRay = Ray;
+		bStampPending = true;
+	}
+}
+
+
+
+void UMeshVertexPaintingTool::OnEndDrag(const FRay& Ray)
+{
+	FinishPainting();
+	bStampPending = false;
+	bInDrag = false;
+}
+
+bool UMeshVertexPaintingTool::HitTest(const FRay& Ray, FHitResult& OutHit)
+{
+	bool bUsed = false;
+	if (UMeshToolManager* MeshToolManager = Cast<UMeshToolManager>(GetToolManager()))
+	{
+		bUsed = MeshToolManager->FindHitResult(Ray, OutHit);
+		LastBestHitResult = OutHit;
+	}
+	return bUsed;
+}
+
+void UMeshVertexPaintingTool::FinishPainting()
+{
+	if (bArePainting)
+	{
+		bArePainting = false;
+		GetToolManager()->EndUndoTransaction();
+		OnPaintingFinishedDelegate.ExecuteIfBound();
+	}
+}
+
+
+UMeshColorPaintingToolProperties::UMeshColorPaintingToolProperties()
+	:UMeshVertexPaintingToolProperties(),
+	bWriteRed(true),
+	bWriteGreen(true),
+	bWriteBlue(true),
+	bWriteAlpha(false)
+{
+
+}
+
+void UMeshColorPaintingToolProperties::SaveProperties(UInteractiveTool* SaveFromTool)
+{
+	Super::SaveProperties(SaveFromTool);
+	UMeshColorPaintingToolProperties* PropertyCache = GetPropertyCache<UMeshColorPaintingToolProperties>();
+	PropertyCache->bWriteRed = this->bWriteRed;
+	PropertyCache->bWriteGreen = this->bWriteGreen;
+	PropertyCache->bWriteBlue = this->bWriteBlue;
+	PropertyCache->bWriteRed = this->bWriteRed;
+}
+
+void UMeshColorPaintingToolProperties::RestoreProperties(UInteractiveTool* RestoreToTool)
+{
+	Super::RestoreProperties(RestoreToTool);
+	UMeshColorPaintingToolProperties* PropertyCache = GetPropertyCache<UMeshColorPaintingToolProperties>();
+	this->bWriteRed = PropertyCache->bWriteRed;
+	this->bWriteGreen = PropertyCache->bWriteGreen;
+	this->bWriteBlue = PropertyCache->bWriteBlue;
+	this->bWriteRed = PropertyCache->bWriteRed;
+}
+
+UMeshColorPaintingTool::UMeshColorPaintingTool()
+{
+	PropertyClass = UMeshColorPaintingToolProperties::StaticClass();
+}
+
+void UMeshColorPaintingTool::Setup()
+{
+	Super::Setup();
+	ColorProperties = Cast<UMeshColorPaintingToolProperties>(BrushProperties);
+}
+
+void UMeshColorPaintingTool::CacheSelectionData()
+{
+	Super::CacheSelectionData();
+	if (UMeshToolManager* MeshToolManager = Cast<UMeshToolManager>(GetToolManager()))
+	{
+		MeshToolManager->ClearPaintableMeshComponents();
+		// Update(cached) Paint LOD level if necessary
+		ColorProperties->LODIndex = FMath::Min<int32>(ColorProperties->LODIndex, GetMaxLODIndexToPaint());
+		CachedLODIndex = ColorProperties->LODIndex;
+		bCachedForceLOD = ColorProperties->bPaintOnSpecificLOD;
+		//Determine LOD level to use for painting(can only paint on LODs in vertex mode)
+		const int32 PaintLODIndex = ColorProperties->bPaintOnSpecificLOD ? ColorProperties->LODIndex : 0;
+		//Determine UV channel to use while painting textures
+		const int32 UVChannel = 0;
+
+		MeshToolManager->CacheSelectionData(PaintLODIndex, UVChannel);
+	}
+}
+
+void UMeshColorPaintingTool::SetAdditionalPaintParameters(FMeshPaintParameters& InPaintParameters)
+{
+	InPaintParameters.bWriteRed = ColorProperties->bWriteRed;
+	InPaintParameters.bWriteGreen = ColorProperties->bWriteGreen;
+	InPaintParameters.bWriteBlue = ColorProperties->bWriteBlue;
+	InPaintParameters.bWriteAlpha = ColorProperties->bWriteAlpha;
+	InPaintParameters.ApplyVertexDataDelegate.AddStatic(&UMeshPaintingToolset::ApplyVertexColorPaint);
+}
+
+int32 UMeshColorPaintingTool::GetMaxLODIndexToPaint() const
+{
+	//The maximum LOD we can paint is decide by the lowest number of LOD in the selection
+	int32 LODMin = TNumericLimits<int32>::Max();
+
+	if (UMeshToolManager* MeshToolManager = Cast<UMeshToolManager>(GetToolManager()))
+	{
+		TArray<UMeshComponent*> SelectedComponents = MeshToolManager->GetSelectedMeshComponents();
+
+		for (UMeshComponent* MeshComponent : SelectedComponents)
+		{
+			int32 NumMeshLODs = 0;
+			if (UMeshPaintingToolset::TryGetNumberOfLODs(MeshComponent, NumMeshLODs))
+			{
+				ensure(NumMeshLODs > 0);
+				LODMin = FMath::Min(LODMin, NumMeshLODs - 1);
+			}
+		}
+		if (LODMin == TNumericLimits<int32>::Max())
+		{
+			LODMin = 1;
+		}
+	}
+	return LODMin;
+}
+
+
+void UMeshColorPaintingTool::LODPaintStateChanged(const bool bLODPaintingEnabled)
+{
+	bool AbortChange = false;
+
+	// Set actual flag in the settings struct
+	ColorProperties->bPaintOnSpecificLOD = bLODPaintingEnabled;
+
+	if (!bLODPaintingEnabled)
+	{
+		// Reset painting LOD index
+		ColorProperties->LODIndex = 0;
+
+	}
+
+	ApplyForcedLODIndex(bLODPaintingEnabled ? CachedLODIndex : -1);
+
+	if (UMeshToolManager* MeshToolManager = Cast<UMeshToolManager>(GetToolManager()))
+	{
+		TArray<UMeshComponent*> PaintableComponents = MeshToolManager->GetPaintableMeshComponents();
+
+		TUniquePtr< FComponentReregisterContext > ComponentReregisterContext;
+		//Make sure all static mesh render is dirty since we change the force LOD
+		for (UMeshComponent* SelectedComponent : PaintableComponents)
+		{
+			if (SelectedComponent)
+			{
+				ComponentReregisterContext.Reset(new FComponentReregisterContext(SelectedComponent));
+			}
+		}
+
+		MeshToolManager->Refresh();
+	}
+}
+
+void UMeshColorPaintingTool::ApplyForcedLODIndex(int32 ForcedLODIndex)
+{
+	if (UMeshToolManager* MeshToolManager = Cast<UMeshToolManager>(GetToolManager()))
+	{
+		TArray<UMeshComponent*> PaintableComponents = MeshToolManager->GetPaintableMeshComponents();
+
+		for (UMeshComponent* SelectedComponent : PaintableComponents)
+		{
+			if (SelectedComponent)
+			{
+				UMeshPaintingToolset::ForceRenderMeshLOD(SelectedComponent, ForcedLODIndex);
+			}
+		}
+	}
+}
+
+
+void UMeshColorPaintingTool::PaintLODChanged()
+{
+	// Enforced LOD for painting
+	if (CachedLODIndex != ColorProperties->LODIndex)
+	{
+		CachedLODIndex = ColorProperties->LODIndex;
+		ApplyForcedLODIndex(bCachedForceLOD ? CachedLODIndex : -1);
+
+		TUniquePtr< FComponentReregisterContext > ComponentReregisterContext;
+		//Make sure all static mesh render is dirty since we change the force LOD
+		if (UMeshToolManager* MeshToolManager = Cast<UMeshToolManager>(GetToolManager()))
+		{
+			TArray<UMeshComponent*> PaintableComponents = MeshToolManager->GetPaintableMeshComponents();
+
+			for (UMeshComponent* SelectedComponent : PaintableComponents)
+			{
+				if (SelectedComponent)
+				{
+					ComponentReregisterContext.Reset(new FComponentReregisterContext(SelectedComponent));
+				}
+			}
+
+			MeshToolManager->Refresh();
+		}
+	}
+}
+
+UMeshWeightPaintingToolProperties::UMeshWeightPaintingToolProperties()
+	:UMeshVertexPaintingToolProperties(),
+	TextureWeightType(EMeshPaintWeightTypes::AlphaLerp),
+	PaintTextureWeightIndex(EMeshPaintTextureIndex::TextureOne),
+	EraseTextureWeightIndex(EMeshPaintTextureIndex::TextureTwo)
+{
+
+}
+
+void UMeshWeightPaintingToolProperties::SaveProperties(UInteractiveTool* SaveFromTool)
+{
+	Super::SaveProperties(SaveFromTool);
+}
+
+void UMeshWeightPaintingToolProperties::RestoreProperties(UInteractiveTool* RestoreToTool)
+{
+	Super::RestoreProperties(RestoreToTool);
+}
+
+UMeshWeightPaintingTool::UMeshWeightPaintingTool()
+{
+	PropertyClass = UMeshWeightPaintingToolProperties::StaticClass();
+}
+
+void UMeshWeightPaintingTool::Setup()
+{
+	Super::Setup();
+	WeightProperties = Cast<UMeshWeightPaintingToolProperties>(BrushProperties);
+}
+
+void UMeshWeightPaintingTool::CacheSelectionData()
+{
+	Super::CacheSelectionData();
+	if (UMeshToolManager* MeshToolManager = Cast<UMeshToolManager>(GetToolManager()))
+	{
+		MeshToolManager->ClearPaintableMeshComponents();
+		//Determine LOD level to use for painting(can only paint on LODs in vertex mode)
+		const int32 PaintLODIndex = 0;
+		//Determine UV channel to use while painting textures
+		const int32 UVChannel = 0;
+		MeshToolManager->CacheSelectionData(PaintLODIndex, UVChannel);
+	}
+}
+
+void UMeshWeightPaintingTool::SetAdditionalPaintParameters(FMeshPaintParameters& InPaintParameters)
+{
+	InPaintParameters.TotalWeightCount = (int32)WeightProperties->TextureWeightType;
+
+	// Select texture weight index based on whether or not we're painting or erasing
+	{
+		const int32 PaintWeightIndex = InPaintParameters.PaintAction == EMeshPaintModeAction::Paint ? (int32)WeightProperties->PaintTextureWeightIndex : (int32)WeightProperties->EraseTextureWeightIndex;
+
+		// Clamp the weight index to fall within the total weight count
+		InPaintParameters.PaintWeightIndex = FMath::Clamp(PaintWeightIndex, 0, InPaintParameters.TotalWeightCount - 1);
+	}
+	InPaintParameters.ApplyVertexDataDelegate.AddStatic(&UMeshPaintingToolset::ApplyVertexWeightPaint);
+}
+
+#undef LOCTEXT_NAMESPACE
\ No newline at end of file
diff --git a/Engine/Plugins/MeshPainting/Source/MeshPaintingToolset/Public/MeshPaintColorBrush.h b/Engine/Plugins/MeshPainting/Source/MeshPaintingToolset/Public/MeshPaintColorBrush.h
deleted file mode 100644
diff --git a/Engine/Plugins/MeshPainting/Source/MeshPaintingToolset/Public/MeshVertexPaintingTool.h b/Engine/Plugins/MeshPainting/Source/MeshPaintingToolset/Public/MeshVertexPaintingTool.h
new file mode 100644
--- /dev/null
+++ b/Engine/Plugins/MeshPainting/Source/MeshPaintingToolset/Public/MeshVertexPaintingTool.h
@@ -0,0 +1,288 @@
+// Copyright 1998-2019 Epic Games, Inc. All Rights Reserved.
+
+#pragma once
+
+#include "CoreMinimal.h"
+#include "UObject/NoExportTypes.h"
+#include "InteractiveToolBuilder.h"
+#include "BaseTools/BaseBrushTool.h"
+#include "InteractiveTool.h"
+#include "Delegates/DelegateCombinations.h"
+#include "Components/MeshComponent.h"
+#include "MeshPaintHelpers.h"
+#include "MeshPaintingToolsetTypes.h"
+#include "MeshVertexPaintingTool.generated.h"
+
+
+struct FToolBuilderState;
+enum class EMeshPaintModeAction : uint8;
+class IMeshPaintGeometryAdapter;
+
+struct FPaintRayResults
+{
+	FMeshPaintParameters Params;
+	FHitResult BestTraceResult;
+};
+
+UENUM()
+enum class EMeshPaintWeightTypes : uint8
+{
+	/** Lerp Between Two Textures using Alpha Value */
+	AlphaLerp = 2 UMETA(DisplayName = "Alpha (Two Textures)"),
+
+	/** Weighting Three Textures according to Channels*/
+	RGB = 3 UMETA(DisplayName = "RGB (Three Textures)"),
+
+	/**  Weighting Four Textures according to Channels*/
+	ARGB = 4 UMETA(DisplayName = "ARGB (Four Textures)"),
+
+	/**  Weighting Five Textures according to Channels */
+	OneMinusARGB = 5 UMETA(DisplayName = "ARGB - 1 (Five Textures)")
+};
+
+UENUM()
+enum class EMeshPaintTextureIndex : uint8
+{
+	TextureOne = 0,
+	TextureTwo,
+	TextureThree,
+	TextureFour,
+	TextureFive
+};
+
+DECLARE_DELEGATE_RetVal(TArray<UMeshComponent*>, FGetSelectedMeshComponents);
+/**
+ *
+ */
+UCLASS()
+class MESHPAINTINGTOOLSET_API UMeshColorPaintingToolBuilder : public UInteractiveToolBuilder
+{
+	GENERATED_BODY()
+
+public:
+	virtual bool CanBuildTool(const FToolBuilderState& SceneState) const override;
+	virtual UInteractiveTool* BuildTool(const FToolBuilderState& SceneState) const override;
+};
+
+UCLASS()
+class MESHPAINTINGTOOLSET_API UMeshWeightPaintingToolBuilder : public UInteractiveToolBuilder
+{
+	GENERATED_BODY()
+
+public:
+	virtual bool CanBuildTool(const FToolBuilderState& SceneState) const override;
+	virtual UInteractiveTool* BuildTool(const FToolBuilderState& SceneState) const override;
+};
+
+
+
+UCLASS()
+class MESHPAINTINGTOOLSET_API UMeshVertexPaintingToolProperties : public UBrushBaseProperties
+{
+	GENERATED_BODY()
+
+public:
+	UMeshVertexPaintingToolProperties();
+
+	/** Color used for Applying Vertex Color Painting */
+	UPROPERTY(EditAnywhere, Category = VertexPainting)
+	FLinearColor PaintColor;
+
+	/** Color used for Erasing Vertex Color Painting */
+	UPROPERTY(EditAnywhere, Category = VertexPainting)
+	FLinearColor EraseColor;
+
+	/** Enables "Flow" painting where paint is continually applied from the brush every tick */
+	UPROPERTY(EditAnywhere, Category = Brush, meta = (DisplayName = "Enable Brush Flow"))
+	bool bEnableFlow;
+
+	/** Whether back-facing triangles should be ignored */
+	UPROPERTY(EditAnywhere, Category = Brush, meta = (DisplayName = "Ignore Back-Facing"))
+	bool bOnlyFrontFacingTriangles;
+
+	/** Size of vertex points drawn when mesh painting is active. */
+	UPROPERTY(EditAnywhere, Category = "VertexPainting|Visualization")
+	float VertexPreviewSize;
+
+	virtual void SaveProperties(UInteractiveTool* SaveFromTool) override;
+	virtual void RestoreProperties(UInteractiveTool* RestoreToTool) override;
+
+};
+
+
+UCLASS()
+class MESHPAINTINGTOOLSET_API UMeshColorPaintingToolProperties : public UMeshVertexPaintingToolProperties
+{
+	GENERATED_BODY()
+
+public:
+	UMeshColorPaintingToolProperties();
+
+
+	/** Whether or not to apply Vertex Color Painting to the Red Channel */
+	UPROPERTY(EditAnywhere, Category = ColorPainting, DisplayName = "Red")
+	bool bWriteRed;
+
+	/** Whether or not to apply Vertex Color Painting to the Green Channel */
+	UPROPERTY(EditAnywhere, Category = ColorPainting, DisplayName = "Green")
+	bool bWriteGreen;
+
+	/** Whether or not to apply Vertex Color Painting to the Blue Channel */
+	UPROPERTY(EditAnywhere, Category = ColorPainting, DisplayName = "Blue")
+	bool bWriteBlue;
+
+	/** Whether or not to apply Vertex Color Painting to the Alpha Channel */
+	UPROPERTY(EditAnywhere, Category = ColorPainting, DisplayName = "Alpha")
+	bool bWriteAlpha;
+
+	/** When unchecked the painting on the base LOD will be propagate automatically to all other LODs when exiting the mode or changing the selection */
+	UPROPERTY(EditAnywhere, Category = Painting)
+	bool bPaintOnSpecificLOD;
+
+	/** LOD Index to which should specifically be painted */
+	UPROPERTY(EditAnywhere, Category = Painting, meta = (UIMin = "0", ClampMin = "0", EditCondition = "bPaintOnSpecificLOD"))
+	int32 LODIndex;
+
+	virtual void SaveProperties(UInteractiveTool* SaveFromTool) override;
+	virtual void RestoreProperties(UInteractiveTool* RestoreToTool) override;
+};
+
+UCLASS()
+class MESHPAINTINGTOOLSET_API UMeshWeightPaintingToolProperties : public UMeshVertexPaintingToolProperties
+{
+	GENERATED_BODY()
+
+public:
+	UMeshWeightPaintingToolProperties();
+
+	/** Texture Blend Weight Painting Mode */
+	UPROPERTY(EditAnywhere, Category = WeightPainting, meta = (EnumCondition = 1))
+	EMeshPaintWeightTypes TextureWeightType;
+
+	/** Texture Blend Weight index which should be applied during Painting */
+	UPROPERTY(EditAnywhere, Category = WeightPainting, meta = (EnumCondition = 1))
+	EMeshPaintTextureIndex PaintTextureWeightIndex;
+
+	/** Texture Blend Weight index which should be erased during Painting */
+	UPROPERTY(EditAnywhere, Category = WeightPainting, meta = (EnumCondition = 1))
+	EMeshPaintTextureIndex EraseTextureWeightIndex;
+
+	virtual void SaveProperties(UInteractiveTool* SaveFromTool) override;
+	virtual void RestoreProperties(UInteractiveTool* RestoreToTool) override;
+};
+
+UCLASS(Abstract)
+class MESHPAINTINGTOOLSET_API UMeshVertexPaintingTool : public UBaseBrushTool
+{
+	GENERATED_BODY()
+
+public:
+	UMeshVertexPaintingTool();
+
+	virtual void Setup() override;
+	virtual void Shutdown(EToolShutdownType ShutdownType) override;
+	virtual void Render(IToolsContextRenderAPI* RenderAPI) override;
+	virtual void Tick(float DeltaTime) override;
+	virtual bool HasCancel() const override { return false; }
+	virtual bool HasAccept() const override;
+	virtual bool CanAccept() const override;
+	virtual void OnBeginDrag(const FRay& Ray) override;
+	virtual void OnUpdateDrag(const FRay& Ray) override;
+	virtual void OnEndDrag(const FRay& Ray) override;
+	virtual	bool HitTest(const FRay& Ray, FHitResult& OutHit) override;
+	virtual void OnPropertyModified(UObject* PropertySet, UProperty* Property) override;
+	virtual bool IsPainting() const
+	{
+		return bArePainting;
+	}
+
+
+	FSimpleDelegate OnPaintingFinished()
+	{
+		return OnPaintingFinishedDelegate;
+	}
+
+protected:
+	virtual void SetAdditionalPaintParameters(FMeshPaintParameters& InPaintParameters) {};
+	virtual void FinishPainting();
+	void UpdateResult();
+	double CalculateTargetEdgeLength(int TargetTriCount);
+	bool Paint(const FVector& InRayOrigin, const FVector& InRayDirection);
+	bool Paint(const TArrayView<TPair<FVector, FVector>>& Rays);
+	virtual void CacheSelectionData() {};
+	/** Per vertex action function used for painting vertex data */
+	void ApplyVertexData(FPerVertexPaintActionArgs& InArgs, int32 VertexIndex, FMeshPaintParameters Parameters);
+
+
+protected:
+	double InitialMeshArea;
+	bool bResultValid;
+	bool bStampPending;
+	bool bInDrag;
+	FRay PendingStampRay;
+
+private:
+	bool PaintInternal(const TArrayView<TPair<FVector, FVector>>& Rays, EMeshPaintModeAction PaintAction, float PaintStrength);
+	
+private:
+	UPROPERTY(Transient)
+	UMeshVertexPaintingToolProperties* VertexProperties;
+	/** Flag for whether or not we are currently painting */
+	bool bArePainting;
+	bool bDoRestoreRenTargets;
+	/** Time kept since the user has started painting */
+	float TimeSinceStartedPainting;
+	/** Overall time value kept for drawing effects */
+	float Time;
+	FHitResult LastBestHitResult;
+	FSimpleDelegate OnPaintingFinishedDelegate;
+};
+
+UCLASS()
+class MESHPAINTINGTOOLSET_API UMeshColorPaintingTool : public UMeshVertexPaintingTool
+{
+	GENERATED_BODY()
+
+public:
+	UMeshColorPaintingTool();
+	virtual void Setup() override;
+	void LODPaintStateChanged(const bool bLODPaintingEnabled);
+	int32 GetMaxLODIndexToPaint() const;
+	void PaintLODChanged();
+	int32 GetCachedLODIndex() const
+	{
+		return CachedLODIndex;
+	}
+
+protected:
+	virtual void CacheSelectionData() override;
+	virtual void SetAdditionalPaintParameters(FMeshPaintParameters& InPaintParameters) override;
+	void ApplyForcedLODIndex(int32 ForcedLODIndex);
+
+private:
+	UPROPERTY(Transient)
+	UMeshColorPaintingToolProperties* ColorProperties;
+	/** Current LOD index used for painting / forcing */
+	int32 CachedLODIndex;
+	/** Whether or not a specific LOD level should be forced */
+	bool bCachedForceLOD;
+};
+
+
+UCLASS()
+class MESHPAINTINGTOOLSET_API UMeshWeightPaintingTool : public UMeshVertexPaintingTool
+{
+	GENERATED_BODY()
+
+public:
+	UMeshWeightPaintingTool();
+	virtual void Setup() override;
+
+protected:
+	virtual void CacheSelectionData() override;
+	virtual void SetAdditionalPaintParameters(FMeshPaintParameters& InPaintParameters);
+
+private:
+	UPROPERTY(Transient)
+	UMeshWeightPaintingToolProperties* WeightProperties;
+};
\ No newline at end of file
