Fixing import options class naming
#rnx

diff --git a/Engine/Plugins/MeshPainting/Source/MeshPaintEditorMode/Private/ImportVertexColorOptionsCustomization.cpp b/Engine/Plugins/MeshPainting/Source/MeshPaintEditorMode/Private/ImportVertexColorOptionsCustomization.cpp
--- a/Engine/Plugins/MeshPainting/Source/MeshPaintEditorMode/Private/ImportVertexColorOptionsCustomization.cpp
+++ b/Engine/Plugins/MeshPainting/Source/MeshPaintEditorMode/Private/ImportVertexColorOptionsCustomization.cpp
@@ -10,7 +10,7 @@
 
 #define LOCTEXT_NAMESPACE "VertexColorImportOptionsCustomization"
 
-void FVertexColorImportOptionsCustomization::CustomizeDetails(IDetailLayoutBuilder& DetailBuilder)
+void FImportVertexColorOptionsCustomization::CustomizeDetails(IDetailLayoutBuilder& DetailBuilder)
 {
 	UMeshVertexColorImportOptions* Options = GetMutableDefault<UMeshVertexColorImportOptions>();
 
@@ -114,7 +114,7 @@ void FVertexColorImportOptionsCustomization::CustomizeDetails(IDetailLayoutBuild
 	}
 }
 
-TSharedRef<SHorizontalBox> FVertexColorImportOptionsCustomization::CreateColorChannelWidget(TSharedRef<IPropertyHandle> ChannelProperty)
+TSharedRef<SHorizontalBox> FImportVertexColorOptionsCustomization::CreateColorChannelWidget(TSharedRef<IPropertyHandle> ChannelProperty)
 {
 	return SNew(SHorizontalBox)
 		+ SHorizontalBox::Slot()
diff --git a/Engine/Plugins/MeshPainting/Source/MeshPaintEditorMode/Private/ImportVertexColorOptionsCustomization.h b/Engine/Plugins/MeshPainting/Source/MeshPaintEditorMode/Private/ImportVertexColorOptionsCustomization.h
--- a/Engine/Plugins/MeshPainting/Source/MeshPaintEditorMode/Private/ImportVertexColorOptionsCustomization.h
+++ b/Engine/Plugins/MeshPainting/Source/MeshPaintEditorMode/Private/ImportVertexColorOptionsCustomization.h
@@ -6,12 +6,12 @@
 #include "IDetailCustomization.h"
 
 /** Customization for importing vertex colors from a texture see SImportVertexColorOptions */
-class FVertexColorImportOptionsCustomization : public IDetailCustomization
+class FImportVertexColorOptionsCustomization : public IDetailCustomization
 {
 public:
 	static TSharedRef<IDetailCustomization> MakeInstance()
 	{
-		return MakeShareable(new FVertexColorImportOptionsCustomization);
+		return MakeShareable(new FImportVertexColorOptionsCustomization);
 	}
 
 	/** IDetailCustomization interface */
