Renaming options for consistency
#rnx

diff --git a/Engine/Plugins/MeshPainting/Source/MeshPaintEditorMode/Private/ImportVertexColorOptions.cpp b/Engine/Plugins/MeshPainting/Source/MeshPaintEditorMode/Private/ImportVertexColorOptions.cpp
--- a/Engine/Plugins/MeshPainting/Source/MeshPaintEditorMode/Private/ImportVertexColorOptions.cpp
+++ b/Engine/Plugins/MeshPainting/Source/MeshPaintEditorMode/Private/ImportVertexColorOptions.cpp
@@ -14,11 +14,11 @@
 
 #define LOCTEXT_NAMESPACE "VertexColorImportOptions"
 
-void SImportVertexColorOptions::Construct(const FArguments& InArgs)
+void SImportVertexColorOptionsWindow::Construct(const FArguments& InArgs)
 {
 	WidgetWindow = InArgs._WidgetWindow;
 
-	Options = GetMutableDefault<UMeshVertexColorImportOptions>();
+	Options = GetMutableDefault<UImportVertexColorOptions>();
 	Options->LODIndex = 0;
 	Options->UVIndex = 0;
 	
@@ -90,7 +90,7 @@ void SImportVertexColorOptions::Construct(const FArguments& InArgs)
 				SAssignNew(ImportButton, SButton)
 				.HAlign(HAlign_Center)
 				.Text(LOCTEXT("VertexColorOptionWindow_Import", "Import"))
-				.OnClicked(this, &SImportVertexColorOptions::OnImport)
+				.OnClicked(this, &SImportVertexColorOptionsWindow::OnImport)
 			]
 			+ SUniformGridPanel::Slot(1, 0)
 			[
@@ -98,7 +98,7 @@ void SImportVertexColorOptions::Construct(const FArguments& InArgs)
 				.HAlign(HAlign_Center)
 				.Text(LOCTEXT("VertexColorOptionWindow_Cancel", "Cancel"))
 				.ToolTipText(LOCTEXT("VertexColorOptionWindow_Cancel_ToolTip", "Cancels importing Vertex Colors"))
-				.OnClicked(this, &SImportVertexColorOptions::OnCancel)
+				.OnClicked(this, &SImportVertexColorOptionsWindow::OnCancel)
 			]
 		]
 	];
diff --git a/Engine/Plugins/MeshPainting/Source/MeshPaintEditorMode/Private/ImportVertexColorOptions.h b/Engine/Plugins/MeshPainting/Source/MeshPaintEditorMode/Private/ImportVertexColorOptions.h
--- a/Engine/Plugins/MeshPainting/Source/MeshPaintEditorMode/Private/ImportVertexColorOptions.h
+++ b/Engine/Plugins/MeshPainting/Source/MeshPaintEditorMode/Private/ImportVertexColorOptions.h
@@ -12,12 +12,12 @@
 #include "ImportVertexColorOptions.generated.h"
 
 UCLASS()
-class UMeshVertexColorImportOptions : public UObject
+class UImportVertexColorOptions : public UObject
 {
 	GENERATED_BODY()
 
 public:
-	UMeshVertexColorImportOptions() : UVIndex(0), LODIndex(0), bRed(true), bBlue(true), bGreen(true), bAlpha(true), bImportToInstance(true), NumLODs(0), bCanImportToInstance(false) {}
+	UImportVertexColorOptions() : UVIndex(0), LODIndex(0), bRed(true), bBlue(true), bGreen(true), bAlpha(true), bImportToInstance(true), NumLODs(0), bCanImportToInstance(false) {}
 	
 	/** Texture Coordinate Channel to use for Sampling the Texture*/
 	UPROPERTY(EditAnywhere, Category=Options)
@@ -68,10 +68,10 @@ public:
 class IDetailsView;
 class UMeshComponent;
 
-class SImportVertexColorOptions : public SCompoundWidget
+class SImportVertexColorOptionsWindow : public SCompoundWidget
 {
 public:
-	SLATE_BEGIN_ARGS(SImportVertexColorOptions)	{}
+	SLATE_BEGIN_ARGS(SImportVertexColorOptionsWindow)	{}
 		SLATE_ARGUMENT(TSharedPtr<SWindow>, WidgetWindow)
 		SLATE_ARGUMENT(FText, FullPath)
 		SLATE_ARGUMENT(UMeshComponent*, Component)
@@ -116,17 +116,17 @@ public:
 		return bShouldImport;
 	}
 
-	UMeshVertexColorImportOptions* GetOptions() const
+	UImportVertexColorOptions* GetOptions() const
 	{
 		return Options;
 	}
 
-	SImportVertexColorOptions()
+	SImportVertexColorOptionsWindow()
 	{}
 private:
 	TWeakPtr< SWindow > WidgetWindow;
 	TSharedPtr< SButton > ImportButton;
 	bool bShouldImport;
 	TSharedPtr<IDetailsView> DetailsView;
-	UMeshVertexColorImportOptions* Options;
+	UImportVertexColorOptions* Options;
 };
\ No newline at end of file
diff --git a/Engine/Plugins/MeshPainting/Source/MeshPaintEditorMode/Private/ImportVertexColorOptionsCustomization.cpp b/Engine/Plugins/MeshPainting/Source/MeshPaintEditorMode/Private/ImportVertexColorOptionsCustomization.cpp
--- a/Engine/Plugins/MeshPainting/Source/MeshPaintEditorMode/Private/ImportVertexColorOptionsCustomization.cpp
+++ b/Engine/Plugins/MeshPainting/Source/MeshPaintEditorMode/Private/ImportVertexColorOptionsCustomization.cpp
@@ -12,15 +12,15 @@
 
 void FImportVertexColorOptionsCustomization::CustomizeDetails(IDetailLayoutBuilder& DetailBuilder)
 {
-	UMeshVertexColorImportOptions* Options = GetMutableDefault<UMeshVertexColorImportOptions>();
+	UImportVertexColorOptions* Options = GetMutableDefault<UImportVertexColorOptions>();
 
 	const FName CategoryName = "Options";
 	IDetailCategoryBuilder& CategoryBuilder = DetailBuilder.EditCategory(CategoryName);
 	
 	/** Retrieve properties we need */
-	TSharedRef<IPropertyHandle> UVProperty = DetailBuilder.GetProperty(GET_MEMBER_NAME_CHECKED(UMeshVertexColorImportOptions, UVIndex));
-	TSharedRef<IPropertyHandle> LODProperty = DetailBuilder.GetProperty(GET_MEMBER_NAME_CHECKED(UMeshVertexColorImportOptions, LODIndex));
-	TSharedRef<IPropertyHandle> InstanceImportProperty = DetailBuilder.GetProperty(GET_MEMBER_NAME_CHECKED(UMeshVertexColorImportOptions, bImportToInstance));
+	TSharedRef<IPropertyHandle> UVProperty = DetailBuilder.GetProperty(GET_MEMBER_NAME_CHECKED(UImportVertexColorOptions, UVIndex));
+	TSharedRef<IPropertyHandle> LODProperty = DetailBuilder.GetProperty(GET_MEMBER_NAME_CHECKED(UImportVertexColorOptions, LODIndex));
+	TSharedRef<IPropertyHandle> InstanceImportProperty = DetailBuilder.GetProperty(GET_MEMBER_NAME_CHECKED(UImportVertexColorOptions, bImportToInstance));
 
 	/** Hide default properties, this allows us to use default widgets */
 	DetailBuilder.HideProperty(UVProperty);
@@ -75,10 +75,10 @@ void FImportVertexColorOptionsCustomization::CustomizeDetails(IDetailLayoutBuild
 	];
 	
 	// Allign all color channels into one widget from left to right
-	TSharedRef<IPropertyHandle> RedChannel = DetailBuilder.GetProperty(GET_MEMBER_NAME_CHECKED(UMeshVertexColorImportOptions, bRed));
-	TSharedRef<IPropertyHandle> GreenChannel = DetailBuilder.GetProperty(GET_MEMBER_NAME_CHECKED(UMeshVertexColorImportOptions, bGreen));
-	TSharedRef<IPropertyHandle> BlueChannel = DetailBuilder.GetProperty(GET_MEMBER_NAME_CHECKED(UMeshVertexColorImportOptions, bBlue));
-	TSharedRef<IPropertyHandle> AlphaChannel = DetailBuilder.GetProperty(GET_MEMBER_NAME_CHECKED(UMeshVertexColorImportOptions, bAlpha));
+	TSharedRef<IPropertyHandle> RedChannel = DetailBuilder.GetProperty(GET_MEMBER_NAME_CHECKED(UImportVertexColorOptions, bRed));
+	TSharedRef<IPropertyHandle> GreenChannel = DetailBuilder.GetProperty(GET_MEMBER_NAME_CHECKED(UImportVertexColorOptions, bGreen));
+	TSharedRef<IPropertyHandle> BlueChannel = DetailBuilder.GetProperty(GET_MEMBER_NAME_CHECKED(UImportVertexColorOptions, bBlue));
+	TSharedRef<IPropertyHandle> AlphaChannel = DetailBuilder.GetProperty(GET_MEMBER_NAME_CHECKED(UImportVertexColorOptions, bAlpha));
 	TArray<TSharedRef<IPropertyHandle>> Channels = { RedChannel, GreenChannel, BlueChannel, AlphaChannel };
 	TSharedPtr<SHorizontalBox> ChannelsWidget;
 
diff --git a/Engine/Plugins/MeshPainting/Source/MeshPaintEditorMode/Private/MeshPaintModeHelpers.cpp b/Engine/Plugins/MeshPainting/Source/MeshPaintEditorMode/Private/MeshPaintModeHelpers.cpp
--- a/Engine/Plugins/MeshPainting/Source/MeshPaintEditorMode/Private/MeshPaintModeHelpers.cpp
+++ b/Engine/Plugins/MeshPainting/Source/MeshPaintEditorMode/Private/MeshPaintModeHelpers.cpp
@@ -183,7 +183,7 @@ void UMeshPaintModeHelpers::ImportVertexColorsFromTexture(UMeshComponent* MeshCo
 				.Title(FText::FromString(TEXT("Vertex Color Import Options")))
 				.SizingRule(ESizingRule::Autosized);
 
-			TSharedPtr<SImportVertexColorOptions> OptionsWindow = SNew(SImportVertexColorOptions).WidgetWindow(Window)
+			TSharedPtr<SImportVertexColorOptionsWindow> OptionsWindow = SNew(SImportVertexColorOptionsWindow).WidgetWindow(Window)
 				.WidgetWindow(Window)
 				.Component(MeshComponent)
 				.FullPath(FText::FromString(ChosenFilename));
@@ -204,7 +204,7 @@ void UMeshPaintModeHelpers::ImportVertexColorsFromTexture(UMeshComponent* MeshCo
 			if (OptionsWindow->ShouldImport())
 			{
 				// Options specified and start importing
-				UMeshVertexColorImportOptions* Options = OptionsWindow->GetOptions();
+				UImportVertexColorOptions* Options = OptionsWindow->GetOptions();
 
 				if (MeshComponent->IsA<UStaticMeshComponent>())
 				{
@@ -249,7 +249,7 @@ void UMeshPaintModeHelpers::ImportVertexColorsFromTexture(UMeshComponent* MeshCo
 }
 
 
-void UMeshPaintModeHelpers::ImportVertexColorsToSkeletalMesh(USkeletalMesh* SkeletalMesh, const UMeshVertexColorImportOptions* Options, UTexture2D* Texture)
+void UMeshPaintModeHelpers::ImportVertexColorsToSkeletalMesh(USkeletalMesh* SkeletalMesh, const UImportVertexColorOptions* Options, UTexture2D* Texture)
 {
 	checkf(SkeletalMesh && Options && Texture, TEXT("Invalid ptr"));
 
@@ -355,7 +355,7 @@ bool UMeshPaintModeHelpers::RetrieveViewportPaintRays(const FSceneView* View, FV
 }
 
 
-void UMeshPaintModeHelpers::ImportVertexColorsToStaticMesh(UStaticMesh* StaticMesh, const UMeshVertexColorImportOptions* Options, UTexture2D* Texture)
+void UMeshPaintModeHelpers::ImportVertexColorsToStaticMesh(UStaticMesh* StaticMesh, const UImportVertexColorOptions* Options, UTexture2D* Texture)
 {
 	checkf(StaticMesh && Options && Texture, TEXT("Invalid ptr"));
 
@@ -400,7 +400,7 @@ void UMeshPaintModeHelpers::ImportVertexColorsToStaticMesh(UStaticMesh* StaticMe
 	StaticMesh->InitResources();
 }
 
-void UMeshPaintModeHelpers::ImportVertexColorsToStaticMeshComponent(UStaticMeshComponent* StaticMeshComponent, const UMeshVertexColorImportOptions* Options, UTexture2D* Texture)
+void UMeshPaintModeHelpers::ImportVertexColorsToStaticMeshComponent(UStaticMeshComponent* StaticMeshComponent, const UImportVertexColorOptions* Options, UTexture2D* Texture)
 {
 	checkf(StaticMeshComponent && Options && Texture, TEXT("Invalid ptr"));
 
diff --git a/Engine/Plugins/MeshPainting/Source/MeshPaintEditorMode/Public/MeshPaintModeHelpers.h b/Engine/Plugins/MeshPainting/Source/MeshPaintEditorMode/Public/MeshPaintModeHelpers.h
--- a/Engine/Plugins/MeshPainting/Source/MeshPaintEditorMode/Public/MeshPaintModeHelpers.h
+++ b/Engine/Plugins/MeshPainting/Source/MeshPaintEditorMode/Public/MeshPaintModeHelpers.h
@@ -9,7 +9,7 @@
 #include "IMeshPaintGeometryAdapter.h"
 
 class FMeshPaintParameters;
-class UMeshVertexColorImportOptions;
+class UImportVertexColorOptions;
 class UTexture2D;
 class UStaticMeshComponent;
 class UStaticMesh;
@@ -43,7 +43,7 @@ public:
 	static void ImportVertexColorsFromTexture(UMeshComponent* MeshComponent);
 
 	/** Imports vertex colors from a Texture to the specified Skeletal Mesh according to user-set options */
-	static void ImportVertexColorsToSkeletalMesh(USkeletalMesh* SkeletalMesh, const UMeshVertexColorImportOptions* Options, UTexture2D* Texture);
+	static void ImportVertexColorsToSkeletalMesh(USkeletalMesh* SkeletalMesh, const UImportVertexColorOptions* Options, UTexture2D* Texture);
 
 	struct FPaintRay
 	{
@@ -57,10 +57,10 @@ public:
 	static bool RetrieveViewportPaintRays(const FSceneView* View, FViewport* Viewport, FPrimitiveDrawInterface* PDI, TArray<FPaintRay>& OutPaintRays);
 
 	/** Imports vertex colors from a Texture to the specified Static Mesh according to user-set options */
-	static void ImportVertexColorsToStaticMesh(UStaticMesh* StaticMesh, const UMeshVertexColorImportOptions* Options, UTexture2D* Texture);
+	static void ImportVertexColorsToStaticMesh(UStaticMesh* StaticMesh, const UImportVertexColorOptions* Options, UTexture2D* Texture);
 
 	/** Imports vertex colors from a Texture to the specified Static Mesh Component according to user-set options */
-	static void ImportVertexColorsToStaticMeshComponent(UStaticMeshComponent* StaticMeshComponent, const UMeshVertexColorImportOptions* Options, UTexture2D* Texture);
+	static void ImportVertexColorsToStaticMeshComponent(UStaticMeshComponent* StaticMeshComponent, const UImportVertexColorOptions* Options, UTexture2D* Texture);
 
 	static void PropagateVertexColors(const TArray<UStaticMeshComponent *> StaticMeshComponents);
 	static bool CanPropagateVertexColors(TArray<UStaticMeshComponent*>& StaticMeshComponents, TArray<UStaticMesh*>& StaticMeshes, int32 NumInstanceVertexColorBytes);
diff --git a/Engine/Plugins/MeshPainting/Source/MeshPaintingToolset/Public/MeshPaintHelpers.h b/Engine/Plugins/MeshPainting/Source/MeshPaintingToolset/Public/MeshPaintHelpers.h
--- a/Engine/Plugins/MeshPainting/Source/MeshPaintingToolset/Public/MeshPaintHelpers.h
+++ b/Engine/Plugins/MeshPainting/Source/MeshPaintingToolset/Public/MeshPaintHelpers.h
@@ -13,7 +13,7 @@
 #include "MeshPaintHelpers.generated.h"
 
 class FMeshPaintParameters;
-class UMeshVertexColorImportOptions;
+class UImportVertexColorOptions;
 class UTexture2D;
 class UStaticMeshComponent;
 class USkeletalMesh;
