Renaming adapters for clarity
#rnx

diff --git a/Engine/Plugins/MeshPainting/Source/MeshPaintEditorMode/Private/MeshPaintMode.cpp b/Engine/Plugins/MeshPainting/Source/MeshPaintEditorMode/Private/MeshPaintMode.cpp
--- a/Engine/Plugins/MeshPainting/Source/MeshPaintEditorMode/Private/MeshPaintMode.cpp
+++ b/Engine/Plugins/MeshPainting/Source/MeshPaintEditorMode/Private/MeshPaintMode.cpp
@@ -306,7 +306,7 @@ void UMeshPaintMode::FillWithVertexColor()
 		Component->Modify();
 		ComponentReregisterContext = MakeUnique<FComponentReregisterContext>(Component);
 
-		TSharedPtr<IMeshPaintGeometryAdapter> MeshAdapter = Cast<UMeshToolManager>(GetToolManager())->GetAdapterForComponent(Component);
+		TSharedPtr<IMeshPaintComponentAdapter> MeshAdapter = Cast<UMeshToolManager>(GetToolManager())->GetAdapterForComponent(Component);
 		if (MeshAdapter)
 		{
 			MeshAdapter->PreEdit();
@@ -589,7 +589,7 @@ void UMeshPaintMode::PropagateVertexColorsToLODs()
 	{
 		if (SelectedComponent)
 		{
-			TSharedPtr<IMeshPaintGeometryAdapter> MeshAdapter = Cast<UMeshToolManager>(GetToolManager())->GetAdapterForComponent(SelectedComponent);
+			TSharedPtr<IMeshPaintComponentAdapter> MeshAdapter = Cast<UMeshToolManager>(GetToolManager())->GetAdapterForComponent(SelectedComponent);
 			UMeshPaintingToolset::ApplyVertexColorsToAllLODs(*MeshAdapter, SelectedComponent);
 			FComponentReregisterContext ReregisterContext(SelectedComponent);
 		}
diff --git a/Engine/Plugins/MeshPainting/Source/MeshPaintEditorMode/Private/MeshPaintMode.h b/Engine/Plugins/MeshPainting/Source/MeshPaintEditorMode/Private/MeshPaintMode.h
--- a/Engine/Plugins/MeshPainting/Source/MeshPaintEditorMode/Private/MeshPaintMode.h
+++ b/Engine/Plugins/MeshPainting/Source/MeshPaintEditorMode/Private/MeshPaintMode.h
@@ -11,7 +11,7 @@
 #include "Framework/Commands/UIAction.h"
 #include "Framework/Commands/UICommandInfo.h"
 #include "Framework/Commands/UICommandList.h"
-#include "IMeshPaintGeometryAdapter.h"
+#include "IMeshPaintComponentAdapter.h"
 #include "MeshPaintHelpers.h"
 #include "MeshVertexPaintingTool.h"
 #include "MeshPaintMode.generated.h"
@@ -19,7 +19,7 @@
 class UMeshVertexPaintingToolProperties;
 class UMeshColorPaintingToolProperties;
 class UMeshPaintModeSettings;
-class IMeshPaintGeometryAdapter;
+class IMeshPaintComponentAdapter;
 class UMeshComponent;
 
 /**
diff --git a/Engine/Plugins/MeshPainting/Source/MeshPaintEditorMode/Private/MeshPaintModeHelpers.cpp b/Engine/Plugins/MeshPainting/Source/MeshPaintEditorMode/Private/MeshPaintModeHelpers.cpp
--- a/Engine/Plugins/MeshPainting/Source/MeshPaintEditorMode/Private/MeshPaintModeHelpers.cpp
+++ b/Engine/Plugins/MeshPainting/Source/MeshPaintEditorMode/Private/MeshPaintModeHelpers.cpp
@@ -5,7 +5,7 @@
 #include "ComponentReregisterContext.h"
 
 #include "MeshPaintModeSettings.h"
-#include "IMeshPaintGeometryAdapter.h"
+#include "IMeshPaintComponentAdapter.h"
 #include "MeshPaintAdapterFactory.h"
 
 #include "Components/StaticMeshComponent.h"
diff --git a/Engine/Plugins/MeshPainting/Source/MeshPaintEditorMode/Public/MeshPaintModeHelpers.h b/Engine/Plugins/MeshPainting/Source/MeshPaintEditorMode/Public/MeshPaintModeHelpers.h
--- a/Engine/Plugins/MeshPainting/Source/MeshPaintEditorMode/Public/MeshPaintModeHelpers.h
+++ b/Engine/Plugins/MeshPainting/Source/MeshPaintEditorMode/Public/MeshPaintModeHelpers.h
@@ -6,7 +6,7 @@
 #include "Engine/EngineTypes.h"
 #include "MeshPaintingToolsetTypes.h"
 #include "Kismet/BlueprintFunctionLibrary.h"
-#include "IMeshPaintGeometryAdapter.h"
+#include "IMeshPaintComponentAdapter.h"
 
 class FMeshPaintParameters;
 class UImportVertexColorOptions;
@@ -14,7 +14,7 @@ class UTexture2D;
 class UStaticMeshComponent;
 class UStaticMesh;
 class USkeletalMesh;
-class IMeshPaintGeometryAdapter;
+class IMeshPaintComponentAdapter;
 class UPaintBrushSettings;
 class FEditorViewportClient;
 class UMeshComponent;
diff --git a/Engine/Plugins/MeshPainting/Source/MeshPaintingToolset/Private/BaseMeshPaintComponentAdapter.cpp b/Engine/Plugins/MeshPainting/Source/MeshPaintingToolset/Private/BaseMeshPaintComponentAdapter.cpp
new file mode 100644
--- /dev/null
+++ b/Engine/Plugins/MeshPainting/Source/MeshPaintingToolset/Private/BaseMeshPaintComponentAdapter.cpp
@@ -0,0 +1,199 @@
+// Copyright 1998-2019 Epic Games, Inc. All Rights Reserved.
+
+#include "BaseMeshPaintComponentAdapter.h"
+#include "MeshAdapter.h"
+#include "Spatial/MeshAABBTree3.h"
+#include "Templates/UniquePtr.h"
+
+bool FBaseMeshPaintComponentAdapter::Initialize()
+{
+	bool bInit = InitializeVertexData() && BuildOctree();
+	FIndexMeshArrayAdapterd TempAdapter(&MeshVertices, &MeshIndices);
+	Adapter = TempAdapter;
+	AABBTree = MakeUnique< TMeshAABBTree3<const FIndexMeshArrayAdapterd>>(&Adapter);
+	return bInit;
+}
+
+bool FBaseMeshPaintComponentAdapter::BuildOctree()
+{
+	bool bValidOctree = false;
+
+	if (MeshVertices.Num() > 0 && MeshIndices.Num() > 0)
+	{
+		bValidOctree = true;
+		
+		// First determine bounding box of mesh verts
+		FBox Bounds;
+		for (const FVector& Vertex : MeshVertices)
+		{
+			Bounds += Vertex;
+		}
+
+		MeshTriOctree = MakeUnique<FMeshPaintTriangleOctree>(Bounds.GetCenter(), Bounds.GetExtent().GetMax());
+
+		const uint32 NumIndexBufferIndices = MeshIndices.Num();
+		check(NumIndexBufferIndices % 3 == 0);
+		const uint32 NumTriangles = NumIndexBufferIndices / 3;
+
+		for (uint32 TriIndex = 0; TriIndex < NumTriangles; ++TriIndex)
+		{
+			// Grab the vertex indices and points for this triangle
+			FMeshPaintTriangle MeshTri;
+			MeshTri.Vertices[0] = MeshVertices[MeshIndices[TriIndex * 3 + 0]];
+			MeshTri.Vertices[1] = MeshVertices[MeshIndices[TriIndex * 3 + 1]];
+			MeshTri.Vertices[2] = MeshVertices[MeshIndices[TriIndex * 3 + 2]];
+			MeshTri.Normal = FVector::CrossProduct(MeshTri.Vertices[1] - MeshTri.Vertices[0], MeshTri.Vertices[2] - MeshTri.Vertices[0]).GetSafeNormal();
+			MeshTri.Index = TriIndex;
+
+			FBox TriBox;
+			TriBox.Min.X = FMath::Min3(MeshTri.Vertices[0].X, MeshTri.Vertices[1].X, MeshTri.Vertices[2].X);
+			TriBox.Min.Y = FMath::Min3(MeshTri.Vertices[0].Y, MeshTri.Vertices[1].Y, MeshTri.Vertices[2].Y);
+			TriBox.Min.Z = FMath::Min3(MeshTri.Vertices[0].Z, MeshTri.Vertices[1].Z, MeshTri.Vertices[2].Z);
+
+			TriBox.Max.X = FMath::Max3(MeshTri.Vertices[0].X, MeshTri.Vertices[1].X, MeshTri.Vertices[2].X);
+			TriBox.Max.Y = FMath::Max3(MeshTri.Vertices[0].Y, MeshTri.Vertices[1].Y, MeshTri.Vertices[2].Y);
+			TriBox.Max.Z = FMath::Max3(MeshTri.Vertices[0].Z, MeshTri.Vertices[1].Z, MeshTri.Vertices[2].Z);
+
+			MeshTri.BoxCenterAndExtent = FBoxCenterAndExtent(TriBox);
+			MeshTriOctree->AddElement(MeshTri);
+		}
+	}
+
+	return bValidOctree;	
+}
+
+const TArray<FVector>& FBaseMeshPaintComponentAdapter::GetMeshVertices() const
+{
+	return MeshVertices;
+}
+
+const TArray<uint32>& FBaseMeshPaintComponentAdapter::GetMeshIndices() const
+{
+	return MeshIndices;
+}
+
+void FBaseMeshPaintComponentAdapter::GetVertexPosition(int32 VertexIndex, FVector& OutVertex) const
+{
+	OutVertex = MeshVertices[VertexIndex];
+}
+
+TArray<uint32> FBaseMeshPaintComponentAdapter::SphereIntersectTriangles(const float ComponentSpaceSquaredBrushRadius, const FVector& ComponentSpaceBrushPosition, const FVector& ComponentSpaceCameraPosition, const bool bOnlyFrontFacing) const
+{
+	TArray<uint32> OutTriangles;
+
+	// Use a bit of distance bias to make sure that we get all of the overlapping triangles.  We
+	// definitely don't want our brush to be cut off by a hard triangle edge
+	const float SquaredRadiusBias = ComponentSpaceSquaredBrushRadius * 0.025f;
+
+	for (FMeshPaintTriangleOctree::TConstElementBoxIterator<> TriIt(*MeshTriOctree, FBoxCenterAndExtent(ComponentSpaceBrushPosition, FVector(FMath::Sqrt(ComponentSpaceSquaredBrushRadius + SquaredRadiusBias)))); TriIt.HasPendingElements(); TriIt.Advance())
+	{
+		// Check to see if the triangle is front facing
+		const FMeshPaintTriangle & CurrentTri = TriIt.GetCurrentElement();
+		const float SignedPlaneDist = FVector::PointPlaneDist(ComponentSpaceCameraPosition, CurrentTri.Vertices[0], CurrentTri.Normal);
+		if (!bOnlyFrontFacing || SignedPlaneDist < 0.0f)
+		{
+			OutTriangles.Add(CurrentTri.Index);
+		}
+	}
+
+	return OutTriangles;
+}
+
+void FBaseMeshPaintComponentAdapter::GetInfluencedVertexIndices(const float ComponentSpaceSquaredBrushRadius, const FVector& ComponentSpaceBrushPosition, const FVector& ComponentSpaceCameraPosition, const bool bOnlyFrontFacing, TSet<int32> &InfluencedVertices) const
+{
+	// Get a list of (optionally front-facing) triangles that are within a reasonable distance to the brush
+	const TArray<uint32> InfluencedTriangles = SphereIntersectTriangles(
+		ComponentSpaceSquaredBrushRadius,
+		ComponentSpaceBrushPosition,
+		ComponentSpaceCameraPosition,
+		bOnlyFrontFacing);
+
+	// Make sure we're dealing with triangle lists
+	const int32 NumIndexBufferIndices = MeshIndices.Num();
+	check(NumIndexBufferIndices % 3 == 0);
+
+	InfluencedVertices.Reserve(InfluencedTriangles.Num());
+	for (const int32 InfluencedTriangle : InfluencedTriangles)
+	{
+		for (int32 Index = 0; Index < 3; ++Index)
+		{
+			const FVector& VertexPosition = MeshVertices[MeshIndices[InfluencedTriangle * 3 + Index]];
+			if ((VertexPosition - ComponentSpaceBrushPosition).SizeSquared() <= ComponentSpaceSquaredBrushRadius)
+			{
+				InfluencedVertices.Add(MeshIndices[InfluencedTriangle * 3 + Index]);
+			}
+		}
+	}
+}
+
+void FBaseMeshPaintComponentAdapter::GetInfluencedVertexData(const float ComponentSpaceSquaredBrushRadius, const FVector& ComponentSpaceBrushPosition, const FVector& ComponentSpaceCameraPosition, const bool bOnlyFrontFacing, TArray<TPair<int32, FVector>>& OutData) const
+{
+	// Get a list of (optionally front-facing) triangles that are within a reasonable distance to the brush
+	TArray<uint32> InfluencedTriangles = SphereIntersectTriangles(
+		ComponentSpaceSquaredBrushRadius,
+		ComponentSpaceBrushPosition,
+		ComponentSpaceCameraPosition,
+		bOnlyFrontFacing);
+
+	// Make sure we're dealing with triangle lists
+	const int32 NumIndexBufferIndices = MeshIndices.Num();
+	check(NumIndexBufferIndices % 3 == 0);
+
+	OutData.Reserve(InfluencedTriangles.Num() * 3);
+	for(int32 InfluencedTriangle : InfluencedTriangles)
+	{
+		for(int32 Index = 0; Index < 3; ++Index)
+		{
+			const FVector& VertexPosition = MeshVertices[MeshIndices[InfluencedTriangle * 3 + Index]];
+			if((VertexPosition - ComponentSpaceBrushPosition).SizeSquared() <= ComponentSpaceSquaredBrushRadius)
+			{
+				OutData.AddDefaulted();
+				TPair<int32, FVector>& OutPair = OutData.Last();
+				OutPair.Key = MeshIndices[InfluencedTriangle * 3 + Index];
+				OutPair.Value = MeshVertices[OutPair.Key];
+			}
+		}
+	}
+}
+
+TArray<FVector> FBaseMeshPaintComponentAdapter::SphereIntersectVertices(const float ComponentSpaceSquaredBrushRadius, const FVector& ComponentSpaceBrushPosition, const FVector& ComponentSpaceCameraPosition, const bool bOnlyFrontFacing) const
+{
+	// Get list of intersecting triangles with given sphere data
+	const TArray<uint32> IntersectedTriangles = SphereIntersectTriangles(ComponentSpaceSquaredBrushRadius, ComponentSpaceBrushPosition, ComponentSpaceCameraPosition, bOnlyFrontFacing);
+	// Get a list of unique vertices indexed by the influenced triangles
+	TSet<int32> InfluencedVertices;
+	InfluencedVertices.Reserve(IntersectedTriangles.Num());
+	for (int32 IntersectedTriangle : IntersectedTriangles)
+	{
+		InfluencedVertices.Add(MeshIndices[IntersectedTriangle * 3 + 0]);
+		InfluencedVertices.Add(MeshIndices[IntersectedTriangle * 3 + 1]);
+		InfluencedVertices.Add(MeshIndices[IntersectedTriangle * 3 + 2]);
+	}
+
+	TArray<FVector> InRangeVertices;
+	InRangeVertices.Empty(MeshVertices.Num());
+	for (int32 VertexIndex : InfluencedVertices)
+	{
+		const FVector& Vertex = MeshVertices[VertexIndex];
+		if (FVector::DistSquared(ComponentSpaceBrushPosition, Vertex) <= ComponentSpaceSquaredBrushRadius)
+		{
+			InRangeVertices.Add(Vertex);
+		}
+	}
+
+	return InRangeVertices;
+}
+
+bool FBaseMeshPaintComponentAdapter::RayIntersectAdapter(FIndex3i& HitTriangle, FVector& HitPosition, const FVector Start, const FVector End) const
+{
+	int32 HitTriangleID;
+	double NearestT;
+	FRay3d Ray(Start, End);
+	bool bHit = AABBTree->FindNearestHitTriangle(Ray, NearestT, HitTriangleID);
+	if (bHit)
+	{
+		HitPosition = Ray.PointAt(NearestT);
+		HitTriangle = Adapter.GetTriangle(HitTriangleID);
+	}
+	return bHit;
+}
diff --git a/Engine/Plugins/MeshPainting/Source/MeshPaintingToolset/Private/BaseMeshPaintGeometryAdapter.cpp b/Engine/Plugins/MeshPainting/Source/MeshPaintingToolset/Private/BaseMeshPaintGeometryAdapter.cpp
deleted file mode 100644
diff --git a/Engine/Plugins/MeshPainting/Source/MeshPaintingToolset/Private/IMeshPaintComponentAdapter.cpp b/Engine/Plugins/MeshPainting/Source/MeshPaintingToolset/Private/IMeshPaintComponentAdapter.cpp
new file mode 100644
--- /dev/null
+++ b/Engine/Plugins/MeshPainting/Source/MeshPaintingToolset/Private/IMeshPaintComponentAdapter.cpp
@@ -0,0 +1,95 @@
+// Copyright 1998-2019 Epic Games, Inc. All Rights Reserved.
+
+#include "IMeshPaintComponentAdapter.h"
+
+#include "Engine/World.h"
+#include "Components/MeshComponent.h"
+#include "MeshPaintingToolsetTypes.h"
+#include "Materials/MaterialExpressionTextureBase.h"
+#include "Materials/MaterialExpressionTextureSample.h"
+#include "Materials/Material.h"
+#include "Materials/MaterialExpressionTextureCoordinate.h"
+#include "Materials/MaterialExpressionTextureSampleParameter.h"
+
+//////////////////////////////////////////////////////////////////////////
+// IMeshPaintGeometryAdapter
+
+void IMeshPaintComponentAdapter::DefaultApplyOrRemoveTextureOverride(UMeshComponent* InMeshComponent, UTexture* SourceTexture, UTexture* OverrideTexture)
+{
+	const ERHIFeatureLevel::Type FeatureLevel = InMeshComponent->GetWorld()->FeatureLevel;
+
+	// Check all the materials on the mesh to see if the user texture is there
+	int32 MaterialIndex = 0;
+	UMaterialInterface* MaterialToCheck = InMeshComponent->GetMaterial(MaterialIndex);
+	while (MaterialToCheck != nullptr)
+	{
+		const bool bIsTextureUsed = DoesMaterialUseTexture(MaterialToCheck, SourceTexture);
+		if (bIsTextureUsed)
+		{
+			MaterialToCheck->OverrideTexture(SourceTexture, OverrideTexture, FeatureLevel);
+		}
+
+		++MaterialIndex;
+		MaterialToCheck = InMeshComponent->GetMaterial(MaterialIndex);
+	}
+}
+
+void IMeshPaintComponentAdapter::DefaultQueryPaintableTextures(int32 MaterialIndex, const UMeshComponent* MeshComponent, int32& OutDefaultIndex, TArray<struct FPaintableTexture>& InOutTextureList)
+{
+	OutDefaultIndex = INDEX_NONE;
+
+	// We already know the material we are painting on, take it off the static mesh component
+	UMaterialInterface* Material = MeshComponent->GetMaterial(MaterialIndex);
+
+	if (Material != NULL)
+	{
+		FPaintableTexture PaintableTexture;
+		// Find all the unique textures used in the top material level of the selected actor materials
+
+		const TArray<UMaterialExpression*>& Expressions = Material->GetMaterial()->Expressions;
+
+		// Only grab the textures from the top level of samples
+		for (auto ItExpressions = Expressions.CreateConstIterator(); ItExpressions; ItExpressions++)
+		{
+			UMaterialExpressionTextureBase* TextureBase = Cast<UMaterialExpressionTextureBase>(*ItExpressions);
+			if (TextureBase != NULL &&
+				TextureBase->Texture != NULL &&
+				!TextureBase->Texture->IsNormalMap() && 
+				!TextureBase->Texture->VirtualTextureStreaming )
+			{
+				// Default UV channel to index 0. 
+				PaintableTexture = FPaintableTexture(TextureBase->Texture, 0);
+
+				// Texture Samples can have UV's specified, check the first node for whether it has a custom UV channel set. 
+				// We only check the first as the Mesh paint mode does not support painting with UV's modified in the shader.
+				UMaterialExpressionTextureSample* TextureSample = Cast<UMaterialExpressionTextureSample>(*ItExpressions);
+				if (TextureSample != NULL)
+				{
+					UMaterialExpressionTextureCoordinate* TextureCoords = Cast<UMaterialExpressionTextureCoordinate>(TextureSample->Coordinates.Expression);
+					if (TextureCoords != NULL)
+					{
+						// Store the uv channel, this is set when the texture is selected. 
+						PaintableTexture.UVChannelIndex = TextureCoords->CoordinateIndex;
+					}
+
+					// Handle texture parameter expressions
+					UMaterialExpressionTextureSampleParameter* TextureSampleParameter = Cast<UMaterialExpressionTextureSampleParameter>(TextureSample);
+					if (TextureSampleParameter != NULL)
+					{
+						// Grab the overridden texture if it exists.  
+						Material->GetTextureParameterValue(TextureSampleParameter->ParameterName, PaintableTexture.Texture);
+					}
+				}
+
+				// note that the same texture will be added again if its UV channel differs. 
+				int32 TextureIndex = InOutTextureList.AddUnique(PaintableTexture);
+
+				// cache the first default index, if there is no previous info this will be used as the selected texture
+				if ((OutDefaultIndex == INDEX_NONE) && TextureBase->IsDefaultMeshpaintTexture)
+				{
+					OutDefaultIndex = TextureIndex;
+				}
+			}
+		}
+	}
+}
\ No newline at end of file
diff --git a/Engine/Plugins/MeshPainting/Source/MeshPaintingToolset/Private/IMeshPaintGeometryAdapter.cpp b/Engine/Plugins/MeshPainting/Source/MeshPaintingToolset/Private/IMeshPaintGeometryAdapter.cpp
deleted file mode 100644
diff --git a/Engine/Plugins/MeshPainting/Source/MeshPaintingToolset/Private/MeshPaintAdapterFactory.cpp b/Engine/Plugins/MeshPainting/Source/MeshPaintingToolset/Private/MeshPaintAdapterFactory.cpp
--- a/Engine/Plugins/MeshPainting/Source/MeshPaintingToolset/Private/MeshPaintAdapterFactory.cpp
+++ b/Engine/Plugins/MeshPainting/Source/MeshPaintingToolset/Private/MeshPaintAdapterFactory.cpp
@@ -5,9 +5,9 @@
 
 TArray<TSharedPtr<class IMeshPaintGeometryAdapterFactory>> FMeshPaintAdapterFactory::FactoryList;
 
-TSharedPtr<class IMeshPaintGeometryAdapter> FMeshPaintAdapterFactory::CreateAdapterForMesh(UMeshComponent* InComponent, int32 InPaintingMeshLODIndex)
+TSharedPtr<class IMeshPaintComponentAdapter> FMeshPaintAdapterFactory::CreateAdapterForMesh(UMeshComponent* InComponent, int32 InPaintingMeshLODIndex)
 {
-	TSharedPtr<IMeshPaintGeometryAdapter> Result;
+	TSharedPtr<IMeshPaintComponentAdapter> Result;
 
 	for (const auto& Factory : FactoryList)
 	{
diff --git a/Engine/Plugins/MeshPainting/Source/MeshPaintingToolset/Private/MeshPaintHelpers.cpp b/Engine/Plugins/MeshPainting/Source/MeshPaintingToolset/Private/MeshPaintHelpers.cpp
--- a/Engine/Plugins/MeshPainting/Source/MeshPaintingToolset/Private/MeshPaintHelpers.cpp
+++ b/Engine/Plugins/MeshPainting/Source/MeshPaintingToolset/Private/MeshPaintHelpers.cpp
@@ -3,7 +3,7 @@
 #include "MeshPaintHelpers.h"
 #include "ComponentReregisterContext.h"
 #include "MeshPaintingToolsetTypes.h"
-#include "IMeshPaintGeometryAdapter.h"
+#include "IMeshPaintComponentAdapter.h"
 #include "MeshPaintAdapterFactory.h"
 #include "Components/StaticMeshComponent.h"
 #include "Components/SkeletalMeshComponent.h"
@@ -862,7 +862,7 @@ void UMeshPaintingToolset::ForceRenderMeshLOD(UMeshComponent* Component, int32 L
 	}	
 }
 
-void UMeshPaintingToolset::ClearMeshTextureOverrides(const IMeshPaintGeometryAdapter& GeometryInfo, UMeshComponent* InMeshComponent)
+void UMeshPaintingToolset::ClearMeshTextureOverrides(const IMeshPaintComponentAdapter& GeometryInfo, UMeshComponent* InMeshComponent)
 {
 	if (InMeshComponent != nullptr)
 	{
@@ -879,7 +879,7 @@ void UMeshPaintingToolset::ClearMeshTextureOverrides(const IMeshPaintGeometryAda
 	}
 }
 
-void UMeshPaintingToolset::ApplyVertexColorsToAllLODs(IMeshPaintGeometryAdapter& GeometryInfo, UMeshComponent* InMeshComponent)
+void UMeshPaintingToolset::ApplyVertexColorsToAllLODs(IMeshPaintComponentAdapter& GeometryInfo, UMeshComponent* InMeshComponent)
 {
 	if (UStaticMeshComponent* StaticMeshComponent = Cast<UStaticMeshComponent>(InMeshComponent))
 	{
@@ -891,7 +891,7 @@ void UMeshPaintingToolset::ApplyVertexColorsToAllLODs(IMeshPaintGeometryAdapter&
 	}
 }
 
-void UMeshPaintingToolset::ApplyVertexColorsToAllLODs(IMeshPaintGeometryAdapter& GeometryInfo, UStaticMeshComponent* StaticMeshComponent)
+void UMeshPaintingToolset::ApplyVertexColorsToAllLODs(IMeshPaintComponentAdapter& GeometryInfo, UStaticMeshComponent* StaticMeshComponent)
 {
 	// If a static mesh component was found, apply LOD0 painting to all lower LODs.
 	if (!StaticMeshComponent || !StaticMeshComponent->GetStaticMesh())
@@ -1157,7 +1157,7 @@ bool UMeshPaintingToolset::ApplyPerVertexPaintAction(FPerVertexPaintActionArgs&
 	return (InfluencedVertices.Num() > 0);
 }
 
-bool UMeshPaintingToolset::ApplyPerTrianglePaintAction(IMeshPaintGeometryAdapter* Adapter, const FVector& CameraPosition, const FVector& HitPosition, const UMeshVertexPaintingToolProperties* Settings, FPerTrianglePaintAction Action)
+bool UMeshPaintingToolset::ApplyPerTrianglePaintAction(IMeshPaintComponentAdapter* Adapter, const FVector& CameraPosition, const FVector& HitPosition, const UMeshVertexPaintingToolProperties* Settings, FPerTrianglePaintAction Action)
 {
 	// Retrieve components world matrix
 	const FMatrix& ComponentToWorldMatrix = Adapter->GetComponentToWorldMatrix();
@@ -1244,7 +1244,7 @@ struct FVertexColorPropogationOctreeSemantics
 };
 typedef TOctree<FPaintedMeshVertex, FVertexColorPropogationOctreeSemantics> TVertexColorPropogationPosOctree;
 
-void UMeshPaintingToolset::ApplyVertexColorsToAllLODs(IMeshPaintGeometryAdapter& GeometryInfo, USkeletalMeshComponent* SkeletalMeshComponent)
+void UMeshPaintingToolset::ApplyVertexColorsToAllLODs(IMeshPaintComponentAdapter& GeometryInfo, USkeletalMeshComponent* SkeletalMeshComponent)
 {
 	checkf(SkeletalMeshComponent != nullptr, TEXT("Invalid Skeletal Mesh Component"));
 	USkeletalMesh* Mesh = SkeletalMeshComponent->SkeletalMesh;
@@ -1421,18 +1421,18 @@ void UMeshToolManager::Shutdown()
 	SelectedMeshComponents.Empty();
 }
 
-TMap<UMeshComponent*, TSharedPtr<IMeshPaintGeometryAdapter>> UMeshToolManager::GetComponentToAdapterMap() const
+TMap<UMeshComponent*, TSharedPtr<IMeshPaintComponentAdapter>> UMeshToolManager::GetComponentToAdapterMap() const
 {
 	return ComponentToAdapterMap;
 }
 
-TSharedPtr<IMeshPaintGeometryAdapter> UMeshToolManager::GetAdapterForComponent(UMeshComponent* InComponent)
+TSharedPtr<IMeshPaintComponentAdapter> UMeshToolManager::GetAdapterForComponent(UMeshComponent* InComponent)
 {
-	TSharedPtr<IMeshPaintGeometryAdapter>* MeshAdapterPtr = ComponentToAdapterMap.Find(InComponent);
-	return MeshAdapterPtr ? *MeshAdapterPtr : TSharedPtr<IMeshPaintGeometryAdapter>();
+	TSharedPtr<IMeshPaintComponentAdapter>* MeshAdapterPtr = ComponentToAdapterMap.Find(InComponent);
+	return MeshAdapterPtr ? *MeshAdapterPtr : TSharedPtr<IMeshPaintComponentAdapter>();
 }
 
-void UMeshToolManager::AddToComponentToAdapterMap(UMeshComponent* InComponent, TSharedPtr<IMeshPaintGeometryAdapter> InAdapter)
+void UMeshToolManager::AddToComponentToAdapterMap(UMeshComponent* InComponent, TSharedPtr<IMeshPaintComponentAdapter> InAdapter)
 {
 	ComponentToAdapterMap.Add(InComponent, InAdapter);
 }
@@ -1507,7 +1507,7 @@ bool UMeshToolManager::FindHitResult(const FRay Ray, FHitResult& BestTraceResult
 
 		for (UMeshComponent* MeshComponent : PaintableComponents)
 		{
-			TSharedPtr<IMeshPaintGeometryAdapter> MeshAdapter = GetComponentToAdapterMap().FindChecked(MeshComponent);
+			TSharedPtr<IMeshPaintComponentAdapter> MeshAdapter = GetComponentToAdapterMap().FindChecked(MeshComponent);
 
 			// Ray trace
 			FHitResult TraceHitResult(1.0f);
@@ -1553,7 +1553,7 @@ void UMeshToolManager::CacheSelectionData(const int32 PaintLODIndex, const int32
 	bSelectionContainsPerLODColors = false;
 	for (UMeshComponent* MeshComponent : SelectedMeshComponents)
 	{
-		TSharedPtr<IMeshPaintGeometryAdapter> MeshAdapter = FMeshPaintAdapterFactory::CreateAdapterForMesh(MeshComponent, PaintLODIndex);
+		TSharedPtr<IMeshPaintComponentAdapter> MeshAdapter = FMeshPaintAdapterFactory::CreateAdapterForMesh(MeshComponent, PaintLODIndex);
 		if (MeshComponent->IsVisible() && MeshAdapter.IsValid() && MeshAdapter->IsValid())
 		{
 			TUniquePtr< FComponentReregisterContext > ComponentReregisterContext;
diff --git a/Engine/Plugins/MeshPainting/Source/MeshPaintingToolset/Private/MeshPaintSkeletalMeshAdapter.cpp b/Engine/Plugins/MeshPainting/Source/MeshPaintingToolset/Private/MeshPaintSkeletalMeshAdapter.cpp
--- a/Engine/Plugins/MeshPainting/Source/MeshPaintingToolset/Private/MeshPaintSkeletalMeshAdapter.cpp
+++ b/Engine/Plugins/MeshPainting/Source/MeshPaintingToolset/Private/MeshPaintSkeletalMeshAdapter.cpp
@@ -13,12 +13,12 @@
 //////////////////////////////////////////////////////////////////////////
 // FMeshPaintGeometryAdapterForSkeletalMeshes
 
-bool FMeshPaintGeometryAdapterForSkeletalMeshes::Construct(UMeshComponent* InComponent, int32 InMeshLODIndex)
+bool FMeshPaintSkeletalMeshComponentAdapter::Construct(UMeshComponent* InComponent, int32 InMeshLODIndex)
 {
 	SkeletalMeshComponent = Cast<USkeletalMeshComponent>(InComponent);
 	if (SkeletalMeshComponent != nullptr)
 	{
-		SkeletalMeshChangedHandle = SkeletalMeshComponent->RegisterOnSkeletalMeshPropertyChanged(USkeletalMeshComponent::FOnSkeletalMeshPropertyChanged::CreateRaw(this, &FMeshPaintGeometryAdapterForSkeletalMeshes::OnSkeletalMeshChanged));
+		SkeletalMeshChangedHandle = SkeletalMeshComponent->RegisterOnSkeletalMeshPropertyChanged(USkeletalMeshComponent::FOnSkeletalMeshPropertyChanged::CreateRaw(this, &FMeshPaintSkeletalMeshComponentAdapter::OnSkeletalMeshChanged));
 
 		if (SkeletalMeshComponent->SkeletalMesh != nullptr)
 		{
@@ -32,7 +32,7 @@ bool FMeshPaintGeometryAdapterForSkeletalMeshes::Construct(UMeshComponent* InCom
 	return false;
 }
 
-FMeshPaintGeometryAdapterForSkeletalMeshes::~FMeshPaintGeometryAdapterForSkeletalMeshes()
+FMeshPaintSkeletalMeshComponentAdapter::~FMeshPaintSkeletalMeshComponentAdapter()
 {
 	if (SkeletalMeshComponent != nullptr)
 	{
@@ -40,7 +40,7 @@ FMeshPaintGeometryAdapterForSkeletalMeshes::~FMeshPaintGeometryAdapterForSkeleta
 	}
 }
 
-void FMeshPaintGeometryAdapterForSkeletalMeshes::OnSkeletalMeshChanged()
+void FMeshPaintSkeletalMeshComponentAdapter::OnSkeletalMeshChanged()
 {
 	OnRemoved();
 	ReferencedSkeletalMesh = SkeletalMeshComponent->SkeletalMesh;
@@ -51,7 +51,7 @@ void FMeshPaintGeometryAdapterForSkeletalMeshes::OnSkeletalMeshChanged()
 	}	
 }
 
-void FMeshPaintGeometryAdapterForSkeletalMeshes::OnPostMeshCached(USkeletalMesh* SkeletalMesh)
+void FMeshPaintSkeletalMeshComponentAdapter::OnPostMeshCached(USkeletalMesh* SkeletalMesh)
 {
 	if (ReferencedSkeletalMesh == SkeletalMesh)
 	{
@@ -59,7 +59,7 @@ void FMeshPaintGeometryAdapterForSkeletalMeshes::OnPostMeshCached(USkeletalMesh*
 	}
 }
 
-bool FMeshPaintGeometryAdapterForSkeletalMeshes::Initialize()
+bool FMeshPaintSkeletalMeshComponentAdapter::Initialize()
 {
 	check(ReferencedSkeletalMesh == SkeletalMeshComponent->SkeletalMesh);
 
@@ -72,14 +72,14 @@ bool FMeshPaintGeometryAdapterForSkeletalMeshes::Initialize()
 		checkf(ReferencedSkeletalMesh->GetImportedModel()->LODModels.IsValidIndex(MeshLODIndex), TEXT("Invalid Imported Model index for vertex painting"));
 		LODModel = &ReferencedSkeletalMesh->GetImportedModel()->LODModels[MeshLODIndex];
 
-		bInitializationResult = FBaseMeshPaintGeometryAdapter::Initialize();
+		bInitializationResult = FBaseMeshPaintComponentAdapter::Initialize();
 	}
 
 	
 	return bInitializationResult;
 }
 
-bool FMeshPaintGeometryAdapterForSkeletalMeshes::InitializeVertexData()
+bool FMeshPaintSkeletalMeshComponentAdapter::InitializeVertexData()
 {
 	// Retrieve mesh vertex and index data 
 	const int32 NumVertices = LODData->GetNumVertices();
@@ -97,7 +97,7 @@ bool FMeshPaintGeometryAdapterForSkeletalMeshes::InitializeVertexData()
 	return (MeshVertices.Num() >= 0 && MeshIndices.Num() > 0);
 }
 
-void FMeshPaintGeometryAdapterForSkeletalMeshes::InitializeAdapterGlobals()
+void FMeshPaintSkeletalMeshComponentAdapter::InitializeAdapterGlobals()
 {
 	static bool bInitialized = false;
 	if (!bInitialized)
@@ -106,21 +106,21 @@ void FMeshPaintGeometryAdapterForSkeletalMeshes::InitializeAdapterGlobals()
 	}
 }
 
-void FMeshPaintGeometryAdapterForSkeletalMeshes::AddReferencedObjectsGlobals(FReferenceCollector& Collector)
+void FMeshPaintSkeletalMeshComponentAdapter::AddReferencedObjectsGlobals(FReferenceCollector& Collector)
 {
 }
 
-void FMeshPaintGeometryAdapterForSkeletalMeshes::AddReferencedObjects(FReferenceCollector& Collector)
+void FMeshPaintSkeletalMeshComponentAdapter::AddReferencedObjects(FReferenceCollector& Collector)
 {
 	Collector.AddReferencedObject(ReferencedSkeletalMesh);
 	Collector.AddReferencedObject(SkeletalMeshComponent);
 }
 
-void FMeshPaintGeometryAdapterForSkeletalMeshes::CleanupGlobals()
+void FMeshPaintSkeletalMeshComponentAdapter::CleanupGlobals()
 {
 }
 
-void FMeshPaintGeometryAdapterForSkeletalMeshes::OnAdded()
+void FMeshPaintSkeletalMeshComponentAdapter::OnAdded()
 {
 	checkf(SkeletalMeshComponent, TEXT("Invalid SkeletalMesh Component"));
 	checkf(ReferencedSkeletalMesh, TEXT("Invalid reference to Skeletal Mesh"));
@@ -130,10 +130,10 @@ void FMeshPaintGeometryAdapterForSkeletalMeshes::OnAdded()
 	SkeletalMeshComponent->InitAnim(true);
 
 	// Register callback for when the skeletal mesh is cached underneath us
-	ReferencedSkeletalMesh->OnPostMeshCached().AddRaw(this, &FMeshPaintGeometryAdapterForSkeletalMeshes::OnPostMeshCached);
+	ReferencedSkeletalMesh->OnPostMeshCached().AddRaw(this, &FMeshPaintSkeletalMeshComponentAdapter::OnPostMeshCached);
 }
 
-void FMeshPaintGeometryAdapterForSkeletalMeshes::OnRemoved()
+void FMeshPaintSkeletalMeshComponentAdapter::OnRemoved()
 {
 	checkf(SkeletalMeshComponent, TEXT("Invalid SkeletalMesh Component"));
 	
@@ -150,7 +150,7 @@ void FMeshPaintGeometryAdapterForSkeletalMeshes::OnRemoved()
 	ReferencedSkeletalMesh->OnPostMeshCached().RemoveAll(this);
 }
 
-bool FMeshPaintGeometryAdapterForSkeletalMeshes::LineTraceComponent(struct FHitResult& OutHit, const FVector Start, const FVector End, const struct FCollisionQueryParams& Params) const
+bool FMeshPaintSkeletalMeshComponentAdapter::LineTraceComponent(struct FHitResult& OutHit, const FVector Start, const FVector End, const struct FCollisionQueryParams& Params) const
 {
 	const bool bHitBounds = FMath::LineSphereIntersection(Start, End.GetSafeNormal(), (End - Start).SizeSquared(), SkeletalMeshComponent->Bounds.Origin, SkeletalMeshComponent->Bounds.SphereRadius);
 	const float SqrRadius = FMath::Square(SkeletalMeshComponent->Bounds.SphereRadius);
@@ -217,22 +217,22 @@ bool FMeshPaintGeometryAdapterForSkeletalMeshes::LineTraceComponent(struct FHitR
 	return bHitTriangle;
 }
 
-void FMeshPaintGeometryAdapterForSkeletalMeshes::QueryPaintableTextures(int32 MaterialIndex, int32& OutDefaultIndex, TArray<struct FPaintableTexture>& InOutTextureList)
+void FMeshPaintSkeletalMeshComponentAdapter::QueryPaintableTextures(int32 MaterialIndex, int32& OutDefaultIndex, TArray<struct FPaintableTexture>& InOutTextureList)
 {
 	DefaultQueryPaintableTextures(MaterialIndex, SkeletalMeshComponent, OutDefaultIndex, InOutTextureList);
 }
 
-void FMeshPaintGeometryAdapterForSkeletalMeshes::ApplyOrRemoveTextureOverride(UTexture* SourceTexture, UTexture* OverrideTexture) const
+void FMeshPaintSkeletalMeshComponentAdapter::ApplyOrRemoveTextureOverride(UTexture* SourceTexture, UTexture* OverrideTexture) const
 {
 	DefaultApplyOrRemoveTextureOverride(SkeletalMeshComponent, SourceTexture, OverrideTexture);
 }
 
-void FMeshPaintGeometryAdapterForSkeletalMeshes::GetTextureCoordinate(int32 VertexIndex, int32 ChannelIndex, FVector2D& OutTextureCoordinate) const
+void FMeshPaintSkeletalMeshComponentAdapter::GetTextureCoordinate(int32 VertexIndex, int32 ChannelIndex, FVector2D& OutTextureCoordinate) const
 {
 	OutTextureCoordinate = LODData->StaticVertexBuffers.StaticMeshVertexBuffer.GetVertexUV(VertexIndex, ChannelIndex);
 }
 
-void FMeshPaintGeometryAdapterForSkeletalMeshes::PreEdit()
+void FMeshPaintSkeletalMeshComponentAdapter::PreEdit()
 {
 	FlushRenderingCommands();
 
@@ -272,13 +272,13 @@ void FMeshPaintGeometryAdapterForSkeletalMeshes::PreEdit()
 	}
 }
 
-void FMeshPaintGeometryAdapterForSkeletalMeshes::PostEdit()
+void FMeshPaintSkeletalMeshComponentAdapter::PostEdit()
 {
 	TUniquePtr< FSkinnedMeshComponentRecreateRenderStateContext > RecreateRenderStateContext = MakeUnique<FSkinnedMeshComponentRecreateRenderStateContext>(ReferencedSkeletalMesh);
 	ReferencedSkeletalMesh->InitResources();
 }
 
-void FMeshPaintGeometryAdapterForSkeletalMeshes::GetVertexColor(int32 VertexIndex, FColor& OutColor, bool bInstance /*= true*/) const
+void FMeshPaintSkeletalMeshComponentAdapter::GetVertexColor(int32 VertexIndex, FColor& OutColor, bool bInstance /*= true*/) const
 {
 	if (LODData->StaticVertexBuffers.ColorVertexBuffer.GetNumVertices() > 0)
 	{
@@ -287,7 +287,7 @@ void FMeshPaintGeometryAdapterForSkeletalMeshes::GetVertexColor(int32 VertexInde
 	}
 }
 
-void FMeshPaintGeometryAdapterForSkeletalMeshes::SetVertexColor(int32 VertexIndex, FColor Color, bool bInstance /*= true*/)
+void FMeshPaintSkeletalMeshComponentAdapter::SetVertexColor(int32 VertexIndex, FColor Color, bool bInstance /*= true*/)
 {
 	if (LODData->StaticVertexBuffers.ColorVertexBuffer.GetNumVertices() > 0)
 	{
@@ -305,7 +305,7 @@ void FMeshPaintGeometryAdapterForSkeletalMeshes::SetVertexColor(int32 VertexInde
 	}
 }
 
-FMatrix FMeshPaintGeometryAdapterForSkeletalMeshes::GetComponentToWorldMatrix() const
+FMatrix FMeshPaintSkeletalMeshComponentAdapter::GetComponentToWorldMatrix() const
 {
 	return SkeletalMeshComponent->GetComponentToWorld().ToMatrixWithScale();
 }
@@ -313,13 +313,13 @@ FMatrix FMeshPaintGeometryAdapterForSkeletalMeshes::GetComponentToWorldMatrix()
 //////////////////////////////////////////////////////////////////////////
 // FMeshPaintGeometryAdapterForSkeletalMeshesFactory
 
-TSharedPtr<IMeshPaintGeometryAdapter> FMeshPaintGeometryAdapterForSkeletalMeshesFactory::Construct(class UMeshComponent* InComponent, int32 InMeshLODIndex) const
+TSharedPtr<IMeshPaintComponentAdapter> FMeshPaintGeometryAdapterForSkeletalMeshesFactory::Construct(class UMeshComponent* InComponent, int32 InMeshLODIndex) const
 {
 	if (USkeletalMeshComponent* SkeletalMeshComponent = Cast<USkeletalMeshComponent>(InComponent))
 	{
 		if (SkeletalMeshComponent->SkeletalMesh != nullptr)
 		{
-			TSharedRef<FMeshPaintGeometryAdapterForSkeletalMeshes> Result = MakeShareable(new FMeshPaintGeometryAdapterForSkeletalMeshes());
+			TSharedRef<FMeshPaintSkeletalMeshComponentAdapter> Result = MakeShareable(new FMeshPaintSkeletalMeshComponentAdapter());
 			if (Result->Construct(InComponent, InMeshLODIndex))
 			{
 				return Result;
diff --git a/Engine/Plugins/MeshPainting/Source/MeshPaintingToolset/Private/MeshPaintSplineMeshAdapter.cpp b/Engine/Plugins/MeshPainting/Source/MeshPaintingToolset/Private/MeshPaintSplineMeshAdapter.cpp
--- a/Engine/Plugins/MeshPainting/Source/MeshPaintingToolset/Private/MeshPaintSplineMeshAdapter.cpp
+++ b/Engine/Plugins/MeshPainting/Source/MeshPaintingToolset/Private/MeshPaintSplineMeshAdapter.cpp
@@ -8,7 +8,7 @@
 //////////////////////////////////////////////////////////////////////////
 // FMeshPaintGeometryAdapterForSplineMeshes
 
-bool FMeshPaintGeometryAdapterForSplineMeshes::InitializeVertexData()
+bool FMeshPaintSplineMeshComponentAdapter::InitializeVertexData()
 {
 	// Cache deformed spline mesh vertices for quick lookup during painting / previewing
 	USplineMeshComponent* SplineMeshComponent = Cast<USplineMeshComponent>(StaticMeshComponent);
@@ -50,13 +50,13 @@ bool FMeshPaintGeometryAdapterForSplineMeshes::InitializeVertexData()
 //////////////////////////////////////////////////////////////////////////
 // FMeshPaintGeometryAdapterForSplineMeshesFactory
 
-TSharedPtr<IMeshPaintGeometryAdapter> FMeshPaintGeometryAdapterForSplineMeshesFactory::Construct(class UMeshComponent* InComponent, int32 MeshLODIndex) const
+TSharedPtr<IMeshPaintComponentAdapter> FMeshPaintGeometryAdapterForSplineMeshesFactory::Construct(class UMeshComponent* InComponent, int32 MeshLODIndex) const
 {
 	if (USplineMeshComponent* SplineMeshComponent = Cast<USplineMeshComponent>(InComponent))
 	{
 		if (SplineMeshComponent->GetStaticMesh() != nullptr)
 		{
-			TSharedRef<FMeshPaintGeometryAdapterForSplineMeshes> Result = MakeShareable(new FMeshPaintGeometryAdapterForSplineMeshes());
+			TSharedRef<FMeshPaintSplineMeshComponentAdapter> Result = MakeShareable(new FMeshPaintSplineMeshComponentAdapter());
 			if (Result->Construct(InComponent, MeshLODIndex))
 			{
 				return Result;
diff --git a/Engine/Plugins/MeshPainting/Source/MeshPaintingToolset/Private/MeshPaintStaticMeshAdapter.cpp b/Engine/Plugins/MeshPainting/Source/MeshPaintingToolset/Private/MeshPaintStaticMeshAdapter.cpp
--- a/Engine/Plugins/MeshPainting/Source/MeshPaintingToolset/Private/MeshPaintStaticMeshAdapter.cpp
+++ b/Engine/Plugins/MeshPainting/Source/MeshPaintingToolset/Private/MeshPaintStaticMeshAdapter.cpp
@@ -11,20 +11,20 @@
 //////////////////////////////////////////////////////////////////////////
 // FMeshPaintGeometryAdapterForStaticMeshes
 
-bool FMeshPaintGeometryAdapterForStaticMeshes::Construct(UMeshComponent* InComponent, int32 InMeshLODIndex)
+bool FMeshPaintStaticMeshComponentAdapter::Construct(UMeshComponent* InComponent, int32 InMeshLODIndex)
 {
 	StaticMeshComponent = Cast<UStaticMeshComponent>(InComponent);
 	if (StaticMeshComponent != nullptr)
 	{
 #if WITH_EDITOR
-		StaticMeshComponent->OnStaticMeshChanged().AddRaw(this, &FMeshPaintGeometryAdapterForStaticMeshes::OnStaticMeshChanged);
+		StaticMeshComponent->OnStaticMeshChanged().AddRaw(this, &FMeshPaintStaticMeshComponentAdapter::OnStaticMeshChanged);
 #endif
 		if (StaticMeshComponent->GetStaticMesh() != nullptr)
 		{
 			ReferencedStaticMesh = StaticMeshComponent->GetStaticMesh();
 			MeshLODIndex = InMeshLODIndex;
 #if WITH_EDITOR
-			ReferencedStaticMesh->OnPostMeshBuild().AddRaw(this, &FMeshPaintGeometryAdapterForStaticMeshes::OnPostMeshBuild);
+			ReferencedStaticMesh->OnPostMeshBuild().AddRaw(this, &FMeshPaintStaticMeshComponentAdapter::OnPostMeshBuild);
 #endif
 			const bool bSuccess = Initialize();
 			return bSuccess;
@@ -34,7 +34,7 @@ bool FMeshPaintGeometryAdapterForStaticMeshes::Construct(UMeshComponent* InCompo
 	return false;
 }
 
-FMeshPaintGeometryAdapterForStaticMeshes::~FMeshPaintGeometryAdapterForStaticMeshes()
+FMeshPaintStaticMeshComponentAdapter::~FMeshPaintStaticMeshComponentAdapter()
 {
 	if (StaticMeshComponent != nullptr)
 	{
@@ -49,13 +49,13 @@ FMeshPaintGeometryAdapterForStaticMeshes::~FMeshPaintGeometryAdapterForStaticMes
 	}
 }
 
-void FMeshPaintGeometryAdapterForStaticMeshes::OnPostMeshBuild(UStaticMesh* StaticMesh)
+void FMeshPaintStaticMeshComponentAdapter::OnPostMeshBuild(UStaticMesh* StaticMesh)
 {
 	check(StaticMesh == ReferencedStaticMesh);
 	Initialize();
 }
 
-void FMeshPaintGeometryAdapterForStaticMeshes::OnStaticMeshChanged(UStaticMeshComponent* InStaticMeshComponent)
+void FMeshPaintStaticMeshComponentAdapter::OnStaticMeshChanged(UStaticMeshComponent* InStaticMeshComponent)
 {
 	check(StaticMeshComponent == InStaticMeshComponent);
 	OnRemoved();
@@ -66,26 +66,26 @@ void FMeshPaintGeometryAdapterForStaticMeshes::OnStaticMeshChanged(UStaticMeshCo
 	if (ReferencedStaticMesh)
 	{
 #if WITH_EDITOR
-		ReferencedStaticMesh->OnPostMeshBuild().AddRaw(this, &FMeshPaintGeometryAdapterForStaticMeshes::OnPostMeshBuild);
+		ReferencedStaticMesh->OnPostMeshBuild().AddRaw(this, &FMeshPaintStaticMeshComponentAdapter::OnPostMeshBuild);
 #endif
 		Initialize();
 		OnAdded();
 	}
 }
 
-bool FMeshPaintGeometryAdapterForStaticMeshes::Initialize()
+bool FMeshPaintStaticMeshComponentAdapter::Initialize()
 {
 	check(ReferencedStaticMesh == StaticMeshComponent->GetStaticMesh());
 	if (MeshLODIndex < ReferencedStaticMesh->GetNumLODs())
 	{
 		LODModel = &(ReferencedStaticMesh->RenderData->LODResources[MeshLODIndex]);
-		return FBaseMeshPaintGeometryAdapter::Initialize();
+		return FBaseMeshPaintComponentAdapter::Initialize();
 	}
 
 	return false;
 }
 
-bool FMeshPaintGeometryAdapterForStaticMeshes::InitializeVertexData()
+bool FMeshPaintStaticMeshComponentAdapter::InitializeVertexData()
 {
 	// Retrieve mesh vertex and index data 
 	const int32 NumVertices = LODModel->VertexBuffers.PositionVertexBuffer.GetNumVertices();
@@ -109,7 +109,7 @@ bool FMeshPaintGeometryAdapterForStaticMeshes::InitializeVertexData()
 	return (MeshVertices.Num() > 0 && MeshIndices.Num() > 0);
 }
 
-void FMeshPaintGeometryAdapterForStaticMeshes::PostEdit()
+void FMeshPaintStaticMeshComponentAdapter::PostEdit()
 {
 	// Lighting does not need to be invalidated when mesh painting
 	const bool bUnbuildLighting = false;
@@ -140,7 +140,7 @@ void FMeshPaintGeometryAdapterForStaticMeshes::PostEdit()
 	}
 }
 
-void FMeshPaintGeometryAdapterForStaticMeshes::InitializeAdapterGlobals()
+void FMeshPaintStaticMeshComponentAdapter::InitializeAdapterGlobals()
 {
 	static bool bInitialized = false;
 	if (!bInitialized)
@@ -149,30 +149,30 @@ void FMeshPaintGeometryAdapterForStaticMeshes::InitializeAdapterGlobals()
 	}
 }
 
-void FMeshPaintGeometryAdapterForStaticMeshes::AddReferencedObjectsGlobals(FReferenceCollector& Collector)
+void FMeshPaintStaticMeshComponentAdapter::AddReferencedObjectsGlobals(FReferenceCollector& Collector)
 {
 
 }
 
-void FMeshPaintGeometryAdapterForStaticMeshes::CleanupGlobals()
+void FMeshPaintStaticMeshComponentAdapter::CleanupGlobals()
 {
 }
 
-void FMeshPaintGeometryAdapterForStaticMeshes::AddReferencedObjects(FReferenceCollector& Collector)
+void FMeshPaintStaticMeshComponentAdapter::AddReferencedObjects(FReferenceCollector& Collector)
 {
 	Collector.AddReferencedObject(ReferencedStaticMesh);
 	Collector.AddReferencedObject(StaticMeshComponent);
 }
 
 
-void FMeshPaintGeometryAdapterForStaticMeshes::OnAdded()
+void FMeshPaintStaticMeshComponentAdapter::OnAdded()
 {
 	check(StaticMeshComponent);
 	check(ReferencedStaticMesh);
 	check(ReferencedStaticMesh == StaticMeshComponent->GetStaticMesh());
 }
 
-void FMeshPaintGeometryAdapterForStaticMeshes::OnRemoved()
+void FMeshPaintStaticMeshComponentAdapter::OnRemoved()
 {
 	check(StaticMeshComponent);
 	
@@ -184,7 +184,7 @@ void FMeshPaintGeometryAdapterForStaticMeshes::OnRemoved()
 	}
 }
 
-bool FMeshPaintGeometryAdapterForStaticMeshes::LineTraceComponent(struct FHitResult& OutHit, const FVector Start, const FVector End, const struct FCollisionQueryParams& Params) const
+bool FMeshPaintStaticMeshComponentAdapter::LineTraceComponent(struct FHitResult& OutHit, const FVector Start, const FVector End, const struct FCollisionQueryParams& Params) const
 {
 	// Ray trace
 	const bool bHitBounds = FMath::LineSphereIntersection(Start, End.GetSafeNormal(), (End - Start).SizeSquared(), StaticMeshComponent->Bounds.Origin, StaticMeshComponent->Bounds.SphereRadius);
@@ -252,17 +252,17 @@ bool FMeshPaintGeometryAdapterForStaticMeshes::LineTraceComponent(struct FHitRes
 	return bHitTriangle;
 }
 
-void FMeshPaintGeometryAdapterForStaticMeshes::QueryPaintableTextures(int32 MaterialIndex, int32& OutDefaultIndex, TArray<struct FPaintableTexture>& InOutTextureList)
+void FMeshPaintStaticMeshComponentAdapter::QueryPaintableTextures(int32 MaterialIndex, int32& OutDefaultIndex, TArray<struct FPaintableTexture>& InOutTextureList)
 {
 	DefaultQueryPaintableTextures(MaterialIndex, StaticMeshComponent, OutDefaultIndex, InOutTextureList);
 }
 
-void FMeshPaintGeometryAdapterForStaticMeshes::ApplyOrRemoveTextureOverride(UTexture* SourceTexture, UTexture* OverrideTexture) const
+void FMeshPaintStaticMeshComponentAdapter::ApplyOrRemoveTextureOverride(UTexture* SourceTexture, UTexture* OverrideTexture) const
 {
 	DefaultApplyOrRemoveTextureOverride(StaticMeshComponent, SourceTexture, OverrideTexture);
 }
 
-void FMeshPaintGeometryAdapterForStaticMeshes::GetVertexColor(int32 VertexIndex, FColor& OutColor, bool bInstance /*= true*/) const
+void FMeshPaintStaticMeshComponentAdapter::GetVertexColor(int32 VertexIndex, FColor& OutColor, bool bInstance /*= true*/) const
 {
 	if (bInstance)
 	{
@@ -296,7 +296,7 @@ void FMeshPaintGeometryAdapterForStaticMeshes::GetVertexColor(int32 VertexIndex,
 	}
 }
 
-void FMeshPaintGeometryAdapterForStaticMeshes::SetVertexColor(int32 VertexIndex, FColor Color, bool bInstance /*= true*/)
+void FMeshPaintStaticMeshComponentAdapter::SetVertexColor(int32 VertexIndex, FColor Color, bool bInstance /*= true*/)
 {
 	// Update the mesh!				
 	if (bInstance)
@@ -334,18 +334,18 @@ void FMeshPaintGeometryAdapterForStaticMeshes::SetVertexColor(int32 VertexIndex,
 	}
 }
 
-FMatrix FMeshPaintGeometryAdapterForStaticMeshes::GetComponentToWorldMatrix() const
+FMatrix FMeshPaintStaticMeshComponentAdapter::GetComponentToWorldMatrix() const
 {
 	return StaticMeshComponent->GetComponentToWorld().ToMatrixWithScale();
 }
 
-void FMeshPaintGeometryAdapterForStaticMeshes::GetTextureCoordinate(int32 VertexIndex, int32 ChannelIndex, FVector2D& OutTextureCoordinate) const
+void FMeshPaintStaticMeshComponentAdapter::GetTextureCoordinate(int32 VertexIndex, int32 ChannelIndex, FVector2D& OutTextureCoordinate) const
 {
 	OutTextureCoordinate = LODModel->VertexBuffers.StaticMeshVertexBuffer.GetVertexUV(VertexIndex, ChannelIndex);
 }
 
 
-void FMeshPaintGeometryAdapterForStaticMeshes::PreEdit()
+void FMeshPaintStaticMeshComponentAdapter::PreEdit()
 {
 	const bool bUsingInstancedVertexColors = true; // Currently we are only painting to instances 
 	UStaticMesh* StaticMesh = ReferencedStaticMesh;
@@ -426,13 +426,13 @@ void FMeshPaintGeometryAdapterForStaticMeshes::PreEdit()
 //////////////////////////////////////////////////////////////////////////
 // FMeshPaintGeometryAdapterForStaticMeshesFactory
 
-TSharedPtr<IMeshPaintGeometryAdapter> FMeshPaintGeometryAdapterForStaticMeshesFactory::Construct(class UMeshComponent* InComponent, int32 InMeshLODIndex) const
+TSharedPtr<IMeshPaintComponentAdapter> FMeshPaintGeometryAdapterForStaticMeshesFactory::Construct(class UMeshComponent* InComponent, int32 InMeshLODIndex) const
 {
 	if (UStaticMeshComponent* StaticMeshComponent = Cast<UStaticMeshComponent>(InComponent))
 	{
 		if (StaticMeshComponent->GetStaticMesh() != nullptr)
 		{
-			TSharedRef<FMeshPaintGeometryAdapterForStaticMeshes> Result = MakeShareable(new FMeshPaintGeometryAdapterForStaticMeshes());
+			TSharedRef<FMeshPaintStaticMeshComponentAdapter> Result = MakeShareable(new FMeshPaintStaticMeshComponentAdapter());
 			if (Result->Construct(InComponent, InMeshLODIndex))
 			{
 				return Result;
diff --git a/Engine/Plugins/MeshPainting/Source/MeshPaintingToolset/Private/MeshSelect.cpp b/Engine/Plugins/MeshPainting/Source/MeshPaintingToolset/Private/MeshSelect.cpp
--- a/Engine/Plugins/MeshPainting/Source/MeshPaintingToolset/Private/MeshSelect.cpp
+++ b/Engine/Plugins/MeshPainting/Source/MeshPaintingToolset/Private/MeshSelect.cpp
@@ -7,7 +7,7 @@
 #include "ToolContextInterfaces.h"
 #include "MeshPaintHelpers.h"
 #include "Components/MeshComponent.h"
-#include "IMeshPaintGeometryAdapter.h"
+#include "IMeshPaintComponentAdapter.h"
 #include "MeshPaintAdapterFactory.h"
 #include "EngineUtils.h"
 #include "Editor.h"
@@ -83,7 +83,7 @@ void UMeshClickTool::OnClicked(const FInputDeviceRay& ClickPos)
 				for (UActorComponent* CandidateComponent : CandidateComponents)
 				{
 					UMeshComponent* MeshComponent = Cast<UMeshComponent>(CandidateComponent);
-					TSharedPtr<IMeshPaintGeometryAdapter> MeshAdapter = FMeshPaintAdapterFactory::CreateAdapterForMesh(MeshComponent, 0);
+					TSharedPtr<IMeshPaintComponentAdapter> MeshAdapter = FMeshPaintAdapterFactory::CreateAdapterForMesh(MeshComponent, 0);
 					MeshToolManager->AddPaintableMeshComponent(MeshComponent);
 					MeshToolManager->AddToComponentToAdapterMap(MeshComponent, MeshAdapter);
 				}
@@ -96,7 +96,7 @@ void UMeshClickTool::OnClicked(const FInputDeviceRay& ClickPos)
 		TArray<AActor*> SelectedActors;
 		if (UMeshComponent* MeshComponent = Cast<UMeshComponent>(TraceHitResult.GetComponent()))
 		{
-			TSharedPtr<IMeshPaintGeometryAdapter> MeshAdapter = FMeshPaintAdapterFactory::CreateAdapterForMesh(MeshComponent, 0);
+			TSharedPtr<IMeshPaintComponentAdapter> MeshAdapter = FMeshPaintAdapterFactory::CreateAdapterForMesh(MeshComponent, 0);
 			if (MeshComponent->IsVisible() && MeshAdapter.IsValid() && MeshAdapter->IsValid())
 			{
 				MeshToolManager->AddPaintableMeshComponent(MeshComponent);
diff --git a/Engine/Plugins/MeshPainting/Source/MeshPaintingToolset/Private/MeshVertexPaintingTool.cpp b/Engine/Plugins/MeshPainting/Source/MeshPaintingToolset/Private/MeshVertexPaintingTool.cpp
--- a/Engine/Plugins/MeshPainting/Source/MeshPaintingToolset/Private/MeshVertexPaintingTool.cpp
+++ b/Engine/Plugins/MeshPainting/Source/MeshPaintingToolset/Private/MeshVertexPaintingTool.cpp
@@ -7,7 +7,7 @@
 #include "ToolContextInterfaces.h"
 #include "Components/MeshComponent.h"
 #include "Math/UnrealMathUtility.h"
-#include "IMeshPaintGeometryAdapter.h"
+#include "IMeshPaintComponentAdapter.h"
 #include "ComponentReregisterContext.h"
 #include "MeshPaintAdapterFactory.h"
 #include "ToolDataVisualizer.h"
@@ -130,7 +130,7 @@ void UMeshVertexPaintingTool::Render(IToolsContextRenderAPI* RenderAPI)
 
 		for (UMeshComponent* CurrentComponent : MeshToolManager->GetPaintableMeshComponents())
 		{
-			TSharedPtr<IMeshPaintGeometryAdapter> MeshAdapter = MeshToolManager->GetAdapterForComponent(Cast<UMeshComponent>(CurrentComponent));
+			TSharedPtr<IMeshPaintComponentAdapter> MeshAdapter = MeshToolManager->GetAdapterForComponent(Cast<UMeshComponent>(CurrentComponent));
 
 			if (MeshAdapter->IsValid() /*&& bRenderVertices*/ && MeshAdapter->SupportsVertexPaint())
 			{
@@ -249,7 +249,7 @@ bool UMeshVertexPaintingTool::PaintInternal(const TArrayView<TPair<FVector, FVec
 
 		for (UMeshComponent* MeshComponent : MeshToolManager->GetPaintableMeshComponents())
 		{
-			TSharedPtr<IMeshPaintGeometryAdapter> MeshAdapter = MeshToolManager->GetAdapterForComponent(MeshComponent);
+			TSharedPtr<IMeshPaintComponentAdapter> MeshAdapter = MeshToolManager->GetAdapterForComponent(MeshComponent);
 
 			// Ray trace
 			FHitResult TraceHitResult(1.0f);
@@ -337,7 +337,7 @@ bool UMeshVertexPaintingTool::PaintInternal(const TArrayView<TPair<FVector, FVec
 			UMeshComponent* HoveredComponent = Entry.Key;
 			TArray<int32>& PaintRayResultIds = Entry.Value;
 
-			IMeshPaintGeometryAdapter* MeshAdapter = MeshToolManager->GetAdapterForComponent(HoveredComponent).Get();
+			IMeshPaintComponentAdapter* MeshAdapter = MeshToolManager->GetAdapterForComponent(HoveredComponent).Get();
 			if (!ensure(MeshAdapter))
 			{
 				continue;
diff --git a/Engine/Plugins/MeshPainting/Source/MeshPaintingToolset/Private/TexturePaintHelpers.cpp b/Engine/Plugins/MeshPainting/Source/MeshPaintingToolset/Private/TexturePaintHelpers.cpp
--- a/Engine/Plugins/MeshPainting/Source/MeshPaintingToolset/Private/TexturePaintHelpers.cpp
+++ b/Engine/Plugins/MeshPainting/Source/MeshPaintingToolset/Private/TexturePaintHelpers.cpp
@@ -9,7 +9,7 @@
 #include "Engine/SkeletalMesh.h"
 #include "Rendering/SkeletalMeshRenderData.h"
 
-#include "IMeshPaintGeometryAdapter.h"
+#include "IMeshPaintComponentAdapter.h"
 
 #include "Engine/TextureRenderTarget2D.h"
 #include "CanvasTypes.h"
@@ -483,7 +483,7 @@ bool TexturePaintHelpers::DoesMeshComponentUseTexture(UMeshComponent* MeshCompon
 	return UsedTextures.Contains(Texture);
 }
 
-void TexturePaintHelpers::RetrieveTexturesForComponent(const UMeshComponent* Component, IMeshPaintGeometryAdapter* Adapter, TArray<FPaintableTexture>& OutTextures)
+void TexturePaintHelpers::RetrieveTexturesForComponent(const UMeshComponent* Component, IMeshPaintComponentAdapter* Adapter, TArray<FPaintableTexture>& OutTextures)
 {
 	// Get the materials used by the mesh
 	TArray<UMaterialInterface*> UsedMaterials;
diff --git a/Engine/Plugins/MeshPainting/Source/MeshPaintingToolset/Public/BaseMeshPaintComponentAdapter.h b/Engine/Plugins/MeshPainting/Source/MeshPaintingToolset/Public/BaseMeshPaintComponentAdapter.h
new file mode 100644
--- /dev/null
+++ b/Engine/Plugins/MeshPainting/Source/MeshPaintingToolset/Public/BaseMeshPaintComponentAdapter.h
@@ -0,0 +1,40 @@
+// Copyright 1998-2019 Epic Games, Inc. All Rights Reserved.
+
+#pragma once
+
+#include "IMeshPaintComponentAdapter.h"
+#include "TMeshPaintOctree.h"
+#include "MeshAdapter.h"
+#include "Spatial/MeshAABBTree3.h"
+
+struct FIndex3i;
+typedef TIndexMeshArrayAdapter<uint32, double> FIndexMeshArrayAdapterd;
+
+/** Base mesh paint geometry adapter, handles basic sphere intersection using a Octree */
+class MESHPAINTINGTOOLSET_API FBaseMeshPaintComponentAdapter : public IMeshPaintComponentAdapter
+{
+public:
+	/** Start IMeshPaintGeometryAdapter Overrides */
+	virtual bool Initialize() override;
+	virtual const TArray<FVector>& GetMeshVertices() const override;
+	virtual const TArray<uint32>& GetMeshIndices() const override;
+	virtual void GetVertexPosition(int32 VertexIndex, FVector& OutVertex) const override;
+	virtual TArray<uint32> SphereIntersectTriangles(const float ComponentSpaceSquaredBrushRadius, const FVector& ComponentSpaceBrushPosition, const FVector& ComponentSpaceCameraPosition, const bool bOnlyFrontFacing) const override;
+	virtual void GetInfluencedVertexIndices(const float ComponentSpaceSquaredBrushRadius, const FVector& ComponentSpaceBrushPosition, const FVector& ComponentSpaceCameraPosition, const bool bOnlyFrontFacing, TSet<int32> &InfluencedVertices) const override;
+	virtual void GetInfluencedVertexData(const float ComponentSpaceSquaredBrushRadius, const FVector& ComponentSpaceBrushPosition, const FVector& ComponentSpaceCameraPosition, const bool bOnlyFrontFacing, TArray<TPair<int32, FVector>>& OutData) const override;
+	virtual TArray<FVector> SphereIntersectVertices(const float ComponentSpaceSquaredBrushRadius, const FVector& ComponentSpaceBrushPosition, const FVector& ComponentSpaceCameraPosition, const bool bOnlyFrontFacing) const override;
+	virtual bool RayIntersectAdapter(FIndex3i& HitTriangle, FVector& HitPosition, const FVector Start, const FVector End) const override;
+	/** End IMeshPaintGeometryAdapter Overrides */
+
+	virtual bool InitializeVertexData() = 0;
+protected:
+	bool BuildOctree();
+protected:
+	/** Index and Vertex data populated by derived classes in InitializeVertexData */
+	TArray<FVector> MeshVertices;
+	TArray<uint32> MeshIndices;
+	/** Octree used for reducing the cost of sphere intersecting with triangles / vertices */
+	TUniquePtr<FMeshPaintTriangleOctree> MeshTriOctree;
+	FIndexMeshArrayAdapterd Adapter;
+	TUniquePtr<TMeshAABBTree3<const FIndexMeshArrayAdapterd>> AABBTree;
+};
diff --git a/Engine/Plugins/MeshPainting/Source/MeshPaintingToolset/Public/BaseMeshPaintGeometryAdapter.h b/Engine/Plugins/MeshPainting/Source/MeshPaintingToolset/Public/BaseMeshPaintGeometryAdapter.h
deleted file mode 100644
diff --git a/Engine/Plugins/MeshPainting/Source/MeshPaintingToolset/Public/IMeshPaintComponentAdapter.h b/Engine/Plugins/MeshPainting/Source/MeshPaintingToolset/Public/IMeshPaintComponentAdapter.h
new file mode 100644
--- /dev/null
+++ b/Engine/Plugins/MeshPainting/Source/MeshPaintingToolset/Public/IMeshPaintComponentAdapter.h
@@ -0,0 +1,102 @@
+// Copyright 1998-2019 Epic Games, Inc. All Rights Reserved.
+
+#pragma once
+
+#include "CoreMinimal.h"
+
+class FReferenceCollector;
+class UMeshComponent;
+class UTexture;
+struct FIndex3i;
+
+/**
+ * Interface for a class to provide mesh painting support for a subclass of UMeshComponent
+ */
+class MESHPAINTINGTOOLSET_API IMeshPaintComponentAdapter
+{
+public:
+	/** Constructs the adapter for a specific LOD index using the mesh component */
+	virtual bool Construct(UMeshComponent* InComponent, int32 MeshLODIndex) = 0;
+
+	/** D-tor */
+	virtual ~IMeshPaintComponentAdapter() {}
+	
+	/** Returns whether or not initialization of necessary data was successful */
+	virtual bool Initialize() = 0;
+
+	/** Called when this adapter is created and added to its owner (used for setting up) */
+	virtual void OnAdded() = 0;
+	
+	/** Called when this adapter is removed from its owner and deleted (used for cleaning up) */
+	virtual void OnRemoved() = 0;
+
+	/** Returns whether or not this adapter is in a valid state */
+	virtual bool IsValid() const = 0;	
+	
+	/** Whether or not this adapter supports texture painting */
+	virtual bool SupportsTexturePaint() const = 0;
+
+	/** Whether or not this adapter supports vertex painting */
+	virtual bool SupportsVertexPaint() const = 0;
+
+	/** Returns the result of tracing a line on the component represented by this adapter */
+	virtual bool LineTraceComponent(struct FHitResult& OutHit, const FVector Start, const FVector End, const struct FCollisionQueryParams& Params) const = 0;	
+	
+	/** Retrieves a list of textures which are suitable for painting */
+	virtual void QueryPaintableTextures(int32 MaterialIndex, int32& OutDefaultIndex, TArray<struct FPaintableTexture>& InOutTextureList) = 0;
+
+	/** Applies or removes an override texture to use while rendering the meshes materials */
+	virtual void ApplyOrRemoveTextureOverride(UTexture* SourceTexture, UTexture* OverrideTexture) const = 0;
+
+	/** Reference collecting to prevent GC-ing */
+	virtual void AddReferencedObjects(FReferenceCollector& Collector) = 0;
+			
+	/** Returns the vertices for the current LOD index in the Mesh */
+	virtual const TArray<FVector>& GetMeshVertices() const = 0;
+
+	/** Returns the indices for the current LOD index in the Mesh */
+	virtual const TArray<uint32>& GetMeshIndices() const = 0;
+
+	/** Returns the Vertex Position at Vertex Index from the Mesh */
+	virtual void GetVertexPosition(int32 VertexIndex, FVector& OutVertex) const = 0;
+
+	/** Sets the Vertex Color at Vertex Index to Color inside of the Mesh */
+	virtual void SetVertexColor(int32 VertexIndex, FColor Color, bool bInstance = true) = 0;
+	
+	/** Returns the Texture Coordinate at Vertex Index from the Mesh */
+	virtual void GetTextureCoordinate(int32 VertexIndex, int32 ChannelIndex, FVector2D& OutTextureCoordinate) const = 0;
+
+	/** Returns the Vertex Color at Vertex Index from the Mesh */
+	virtual void GetVertexColor(int32 VertexIndex, FColor& OutColor, bool bInstance = true) const = 0;
+
+	/** Returns the Component to World matrix from the Mesh Component */
+	virtual FMatrix GetComponentToWorldMatrix() const = 0;
+
+	/** Pre Edit to setup necessary data */
+	virtual void PreEdit() = 0;
+
+	/** Post Edit to clear up or update necessary data */
+	virtual void PostEdit() = 0;
+
+	/** Default functionality for applying or remove a texture override */
+	static void DefaultApplyOrRemoveTextureOverride(UMeshComponent* InMeshComponent, UTexture* SourceTexture, UTexture* OverrideTexture);
+
+	/** Default functionality for querying textures from a mesh component which are suitable for texture painting */
+	static void DefaultQueryPaintableTextures(int32 MaterialIndex, const UMeshComponent* MeshComponent, int32& OutDefaultIndex, TArray<struct FPaintableTexture>& InOutTextureList);
+
+	/** Returns the triangle indices which intersect with the given sphere */
+	virtual TArray<uint32> SphereIntersectTriangles(const float ComponentSpaceSquaredBrushRadius, const FVector& ComponentSpaceBrushPosition, const FVector& ComponentSpaceCameraPosition, const bool bOnlyFrontFacing) const = 0;
+
+	/** Retrieves the influenced vertex indices which intersect the given sphere */
+	virtual void GetInfluencedVertexIndices(const float ComponentSpaceSquaredBrushRadius, const FVector& ComponentSpaceBrushPosition, const FVector& ComponentSpaceCameraPosition, const bool bOnlyFrontFacing, TSet<int32> &InfluencedVertices) const = 0;
+	
+	/** Retrieves the influenced vertex indices and positions that intersect the given sphere */
+	virtual void GetInfluencedVertexData(const float ComponentSpaceSquaredBrushRadius, const FVector& ComponentSpaceBrushPosition, const FVector& ComponentSpaceCameraPosition, const bool bOnlyFrontFacing, TArray<TPair<int32, FVector>>& OutData) const = 0;
+
+	/** Returns the vertex positions which intersect the given sphere */
+	virtual TArray<FVector> SphereIntersectVertices(const float ComponentSpaceSquaredBrushRadius, const FVector& ComponentSpaceBrushPosition, const FVector& ComponentSpaceCameraPosition, const bool bOnlyFrontFacing) const = 0;
+
+	/** Returns the hit result which intersect with the given ray */
+	virtual bool RayIntersectAdapter(FIndex3i& HitTriangle, FVector& HitPosition, const FVector Start, const FVector End) const = 0;
+
+};
\ No newline at end of file
diff --git a/Engine/Plugins/MeshPainting/Source/MeshPaintingToolset/Public/IMeshPaintGeometryAdapter.h b/Engine/Plugins/MeshPainting/Source/MeshPaintingToolset/Public/IMeshPaintGeometryAdapter.h
deleted file mode 100644
diff --git a/Engine/Plugins/MeshPainting/Source/MeshPaintingToolset/Public/IMeshPaintGeometryAdapterFactory.h b/Engine/Plugins/MeshPainting/Source/MeshPaintingToolset/Public/IMeshPaintGeometryAdapterFactory.h
--- a/Engine/Plugins/MeshPainting/Source/MeshPaintingToolset/Public/IMeshPaintGeometryAdapterFactory.h
+++ b/Engine/Plugins/MeshPainting/Source/MeshPaintingToolset/Public/IMeshPaintGeometryAdapterFactory.h
@@ -2,7 +2,7 @@
 
 #pragma once
 
-class IMeshPaintGeometryAdapter;
+class IMeshPaintComponentAdapter;
 class UMeshComponent;
 class FReferenceCollector;
 
@@ -12,7 +12,7 @@ class FReferenceCollector;
 class MESHPAINTINGTOOLSET_API IMeshPaintGeometryAdapterFactory
 {
 public:
-	virtual TSharedPtr<IMeshPaintGeometryAdapter> Construct(UMeshComponent* InComponent, int32 InMeshLODIndex) const = 0;
+	virtual TSharedPtr<IMeshPaintComponentAdapter> Construct(UMeshComponent* InComponent, int32 InMeshLODIndex) const = 0;
 	virtual void InitializeAdapterGlobals() = 0;
 	virtual void AddReferencedObjectsGlobals(FReferenceCollector& Collector) = 0;
 	virtual void CleanupGlobals() = 0;
diff --git a/Engine/Plugins/MeshPainting/Source/MeshPaintingToolset/Public/MeshPaintAdapterFactory.h b/Engine/Plugins/MeshPainting/Source/MeshPaintingToolset/Public/MeshPaintAdapterFactory.h
--- a/Engine/Plugins/MeshPainting/Source/MeshPaintingToolset/Public/MeshPaintAdapterFactory.h
+++ b/Engine/Plugins/MeshPainting/Source/MeshPaintingToolset/Public/MeshPaintAdapterFactory.h
@@ -14,7 +14,7 @@ public:
 	static TArray<TSharedPtr<IMeshPaintGeometryAdapterFactory>> FactoryList;
 
 public:
-	static TSharedPtr<class IMeshPaintGeometryAdapter> CreateAdapterForMesh(UMeshComponent* InComponent, int32 InPaintingMeshLODIndex);
+	static TSharedPtr<class IMeshPaintComponentAdapter> CreateAdapterForMesh(UMeshComponent* InComponent, int32 InPaintingMeshLODIndex);
 	static void InitializeAdapterGlobals();
 	static void AddReferencedObjectsGlobals(FReferenceCollector& Collector);
 	static void CleanupGlobals();
diff --git a/Engine/Plugins/MeshPainting/Source/MeshPaintingToolset/Public/MeshPaintHelpers.h b/Engine/Plugins/MeshPainting/Source/MeshPaintingToolset/Public/MeshPaintHelpers.h
--- a/Engine/Plugins/MeshPainting/Source/MeshPaintingToolset/Public/MeshPaintHelpers.h
+++ b/Engine/Plugins/MeshPainting/Source/MeshPaintingToolset/Public/MeshPaintHelpers.h
@@ -6,7 +6,7 @@
 #include "Engine/EngineTypes.h"
 #include "MeshPaintingToolsetTypes.h"
 #include "Kismet/BlueprintFunctionLibrary.h"
-#include "IMeshPaintGeometryAdapter.h"
+#include "IMeshPaintComponentAdapter.h"
 #include "InteractiveToolManager.h"
 #include "InteractiveToolsContext.h"
 #include "Engine/StaticMesh.h"
@@ -17,7 +17,7 @@ class UImportVertexColorOptions;
 class UTexture2D;
 class UStaticMeshComponent;
 class USkeletalMesh;
-class IMeshPaintGeometryAdapter;
+class IMeshPaintComponentAdapter;
 class UPaintBrushSettings;
 class FEditorViewportClient;
 class UMeshComponent;
@@ -109,7 +109,7 @@ enum class ETexturePaintWeightIndex : uint8
 /** Parameters for paint actions, stored together for convenience */
 struct FPerVertexPaintActionArgs
 {
-	IMeshPaintGeometryAdapter* Adapter;
+	IMeshPaintComponentAdapter* Adapter;
 	UMeshVertexPaintingToolProperties* BrushProperties;
 	FVector CameraPosition;
 	FHitResult HitResult;
@@ -118,7 +118,7 @@ struct FPerVertexPaintActionArgs
 
 /** Delegates used to call per-vertex/triangle actions */
 DECLARE_DELEGATE_TwoParams(FPerVertexPaintAction, FPerVertexPaintActionArgs& /*Args*/, int32 /*VertexIndex*/);
-DECLARE_DELEGATE_ThreeParams(FPerTrianglePaintAction, IMeshPaintGeometryAdapter* /*Adapter*/, int32 /*TriangleIndex*/, const int32[3] /*Vertex Indices*/);
+DECLARE_DELEGATE_ThreeParams(FPerTrianglePaintAction, IMeshPaintComponentAdapter* /*Adapter*/, int32 /*TriangleIndex*/, const int32[3] /*Vertex Indices*/);
 
 class MESHPAINTINGTOOLSET_API UMeshPaintingToolset : public UBlueprintFunctionLibrary
 {
@@ -167,16 +167,16 @@ public:
 	static void ForceRenderMeshLOD(UMeshComponent* Component, int32 LODIndex);
 
 	/** Clears all texture overrides for this component. */
-	static void ClearMeshTextureOverrides(const IMeshPaintGeometryAdapter& GeometryInfo, UMeshComponent* InMeshComponent);
+	static void ClearMeshTextureOverrides(const IMeshPaintComponentAdapter& GeometryInfo, UMeshComponent* InMeshComponent);
 
 	/** Applies vertex color painting found on LOD 0 to all lower LODs. */
-	static void ApplyVertexColorsToAllLODs(IMeshPaintGeometryAdapter& GeometryInfo, UMeshComponent* InMeshComponent);
+	static void ApplyVertexColorsToAllLODs(IMeshPaintComponentAdapter& GeometryInfo, UMeshComponent* InMeshComponent);
 
 	/** Applies the vertex colors found in LOD level 0 to all contained LOD levels in the StaticMeshComponent */
-	static void ApplyVertexColorsToAllLODs(IMeshPaintGeometryAdapter& GeometryInfo, UStaticMeshComponent* StaticMeshComponent);
+	static void ApplyVertexColorsToAllLODs(IMeshPaintComponentAdapter& GeometryInfo, UStaticMeshComponent* StaticMeshComponent);
 
 	/** Applies the vertex colors found in LOD level 0 to all contained LOD levels in the SkeletalMeshComponent */
-	static void ApplyVertexColorsToAllLODs(IMeshPaintGeometryAdapter& GeometryInfo, USkeletalMeshComponent* SkeletalMeshComponent);
+	static void ApplyVertexColorsToAllLODs(IMeshPaintComponentAdapter& GeometryInfo, USkeletalMeshComponent* SkeletalMeshComponent);
 
 	/** Returns the number of Mesh LODs for the given MeshComponent */
 	static int32 GetNumberOfLODs(const UMeshComponent* MeshComponent);
@@ -199,7 +199,7 @@ public:
 	static bool GetPerVertexPaintInfluencedVertices(FPerVertexPaintActionArgs& InArgs, TSet<int32>& InfluencedVertices);
 
 	/** Given the adapter, settings and view-information retrieves influences triangles and applies Action to them */
-	static bool ApplyPerTrianglePaintAction(IMeshPaintGeometryAdapter* Adapter, const FVector& CameraPosition, const FVector& HitPosition, const UMeshVertexPaintingToolProperties* Settings, FPerTrianglePaintAction Action);
+	static bool ApplyPerTrianglePaintAction(IMeshPaintComponentAdapter* Adapter, const FVector& CameraPosition, const FVector& HitPosition, const UMeshVertexPaintingToolProperties* Settings, FPerTrianglePaintAction Action);
 
 	/** Applies vertex painting to InOutvertexColor according to the given parameters  */
 	static bool PaintVertex(const FVector& InVertexPosition, const FMeshPaintParameters& InParams, FColor& InOutVertexColor);
@@ -268,9 +268,9 @@ public:
 	virtual void Shutdown() override;
 
 	/** Map of geometry adapters for each selected mesh component */
-	TMap<UMeshComponent*, TSharedPtr<IMeshPaintGeometryAdapter>> GetComponentToAdapterMap() const;
-	TSharedPtr<IMeshPaintGeometryAdapter> GetAdapterForComponent(UMeshComponent* InComponent);
-	void AddToComponentToAdapterMap(UMeshComponent* InComponent, TSharedPtr<IMeshPaintGeometryAdapter> InAdapter);
+	TMap<UMeshComponent*, TSharedPtr<IMeshPaintComponentAdapter>> GetComponentToAdapterMap() const;
+	TSharedPtr<IMeshPaintComponentAdapter> GetAdapterForComponent(UMeshComponent* InComponent);
+	void AddToComponentToAdapterMap(UMeshComponent* InComponent, TSharedPtr<IMeshPaintComponentAdapter> InAdapter);
 
 	TArray<UMeshComponent*> GetSelectedMeshComponents() const;
 	void AddSelectedMeshComponents(const TArray<UMeshComponent*>& InComponents);
@@ -297,7 +297,7 @@ private:
 
 private:
 	/** Map of geometry adapters for each selected mesh component */
-	TMap<UMeshComponent*, TSharedPtr<IMeshPaintGeometryAdapter>> ComponentToAdapterMap;
+	TMap<UMeshComponent*, TSharedPtr<IMeshPaintComponentAdapter>> ComponentToAdapterMap;
 	TArray<UMeshComponent*> SelectedMeshComponents;
 	/** Mesh components within the current selection which are eligible for painting */
 	TArray<UMeshComponent*> PaintableComponents;
diff --git a/Engine/Plugins/MeshPainting/Source/MeshPaintingToolset/Public/MeshPaintSkeletalMeshAdapter.h b/Engine/Plugins/MeshPainting/Source/MeshPaintingToolset/Public/MeshPaintSkeletalMeshAdapter.h
--- a/Engine/Plugins/MeshPainting/Source/MeshPaintingToolset/Public/MeshPaintSkeletalMeshAdapter.h
+++ b/Engine/Plugins/MeshPainting/Source/MeshPaintingToolset/Public/MeshPaintSkeletalMeshAdapter.h
@@ -5,7 +5,7 @@
 #include "CoreMinimal.h"
 #include "RawIndexBuffer.h"
 #include "Components/SkeletalMeshComponent.h"
-#include "BaseMeshPaintGeometryAdapter.h"
+#include "BaseMeshPaintComponentAdapter.h"
 #include "IMeshPaintGeometryAdapterFactory.h"
 
 class UBodySetup;
@@ -19,7 +19,7 @@ class UMeshComponent;
 
 //////////////////////////////////////////////////////////////////////////
 // FMeshPaintGeometryAdapterForSkeletalMeshes
-class MESHPAINTINGTOOLSET_API FMeshPaintGeometryAdapterForSkeletalMeshes : public FBaseMeshPaintGeometryAdapter
+class MESHPAINTINGTOOLSET_API FMeshPaintSkeletalMeshComponentAdapter : public FBaseMeshPaintComponentAdapter
 {
 public:
 	static void InitializeAdapterGlobals();
@@ -27,7 +27,7 @@ public:
 	static void CleanupGlobals();
 	/** Start IMeshPaintGeometryAdapter Overrides */
 	virtual bool Construct(UMeshComponent* InComponent, int32 InMeshLODIndex) override;
-	virtual ~FMeshPaintGeometryAdapterForSkeletalMeshes();
+	virtual ~FMeshPaintSkeletalMeshComponentAdapter();
 	virtual bool Initialize() override;
 	virtual void OnAdded() override;
 	virtual void OnRemoved() override;
@@ -79,8 +79,8 @@ protected:
 class MESHPAINTINGTOOLSET_API FMeshPaintGeometryAdapterForSkeletalMeshesFactory : public IMeshPaintGeometryAdapterFactory
 {
 public:
-	virtual TSharedPtr<IMeshPaintGeometryAdapter> Construct(UMeshComponent* InComponent, int32 InMeshLODIndex) const override;
-	virtual void InitializeAdapterGlobals() override { FMeshPaintGeometryAdapterForSkeletalMeshes::InitializeAdapterGlobals(); }
-	virtual void AddReferencedObjectsGlobals(FReferenceCollector& Collector) override { FMeshPaintGeometryAdapterForSkeletalMeshes::AddReferencedObjectsGlobals(Collector); }
-	virtual void CleanupGlobals() override { FMeshPaintGeometryAdapterForSkeletalMeshes::CleanupGlobals(); }
+	virtual TSharedPtr<IMeshPaintComponentAdapter> Construct(UMeshComponent* InComponent, int32 InMeshLODIndex) const override;
+	virtual void InitializeAdapterGlobals() override { FMeshPaintSkeletalMeshComponentAdapter::InitializeAdapterGlobals(); }
+	virtual void AddReferencedObjectsGlobals(FReferenceCollector& Collector) override { FMeshPaintSkeletalMeshComponentAdapter::AddReferencedObjectsGlobals(Collector); }
+	virtual void CleanupGlobals() override { FMeshPaintSkeletalMeshComponentAdapter::CleanupGlobals(); }
 };
diff --git a/Engine/Plugins/MeshPainting/Source/MeshPaintingToolset/Public/MeshPaintSplineMeshAdapter.h b/Engine/Plugins/MeshPainting/Source/MeshPaintingToolset/Public/MeshPaintSplineMeshAdapter.h
--- a/Engine/Plugins/MeshPainting/Source/MeshPaintingToolset/Public/MeshPaintSplineMeshAdapter.h
+++ b/Engine/Plugins/MeshPainting/Source/MeshPaintingToolset/Public/MeshPaintSplineMeshAdapter.h
@@ -8,7 +8,7 @@
 //////////////////////////////////////////////////////////////////////////
 // FMeshPaintGeometryAdapterForSplineMeshes
 
-class MESHPAINTINGTOOLSET_API FMeshPaintGeometryAdapterForSplineMeshes : public FMeshPaintGeometryAdapterForStaticMeshes
+class MESHPAINTINGTOOLSET_API FMeshPaintSplineMeshComponentAdapter : public FMeshPaintStaticMeshComponentAdapter
 {
 public:
 	virtual bool InitializeVertexData() override;
@@ -20,5 +20,5 @@ public:
 class MESHPAINTINGTOOLSET_API FMeshPaintGeometryAdapterForSplineMeshesFactory : public FMeshPaintGeometryAdapterForStaticMeshesFactory
 {
 public:
-	virtual TSharedPtr<IMeshPaintGeometryAdapter> Construct(class UMeshComponent* InComponent, int32 MeshLODIndex) const override;
+	virtual TSharedPtr<IMeshPaintComponentAdapter> Construct(class UMeshComponent* InComponent, int32 MeshLODIndex) const override;
 };
diff --git a/Engine/Plugins/MeshPainting/Source/MeshPaintingToolset/Public/MeshPaintStaticMeshAdapter.h b/Engine/Plugins/MeshPainting/Source/MeshPaintingToolset/Public/MeshPaintStaticMeshAdapter.h
--- a/Engine/Plugins/MeshPainting/Source/MeshPaintingToolset/Public/MeshPaintStaticMeshAdapter.h
+++ b/Engine/Plugins/MeshPainting/Source/MeshPaintingToolset/Public/MeshPaintStaticMeshAdapter.h
@@ -4,7 +4,7 @@
 
 #include "CoreMinimal.h"
 #include "IMeshPaintGeometryAdapterFactory.h"
-#include "BaseMeshPaintGeometryAdapter.h"
+#include "BaseMeshPaintComponentAdapter.h"
 #include "RawIndexBuffer.h"
 #include "Components/StaticMeshComponent.h"
 
@@ -16,7 +16,7 @@ struct FStaticMeshLODResources;
 //////////////////////////////////////////////////////////////////////////
 // FMeshPaintGeometryAdapterForStaticMeshes
 
-class MESHPAINTINGTOOLSET_API FMeshPaintGeometryAdapterForStaticMeshes : public FBaseMeshPaintGeometryAdapter
+class MESHPAINTINGTOOLSET_API FMeshPaintStaticMeshComponentAdapter : public FBaseMeshPaintComponentAdapter
 {
 public:
 	static void InitializeAdapterGlobals();
@@ -24,7 +24,7 @@ public:
 	static void CleanupGlobals();
 	/** Start IMeshPaintGeometryAdapter Overrides */
 	virtual bool Construct(UMeshComponent* InComponent, int32 InMeshLODIndex) override;
-	virtual ~FMeshPaintGeometryAdapterForStaticMeshes();
+	virtual ~FMeshPaintStaticMeshComponentAdapter();
 	virtual bool Initialize() override;
 	virtual void OnAdded() override;
 	virtual void OnRemoved() override;
@@ -69,8 +69,8 @@ protected:
 class MESHPAINTINGTOOLSET_API FMeshPaintGeometryAdapterForStaticMeshesFactory : public IMeshPaintGeometryAdapterFactory
 {
 public:
-	virtual TSharedPtr<IMeshPaintGeometryAdapter> Construct(class UMeshComponent* InComponent, int32 InMeshLODIndex) const override;
-	virtual void InitializeAdapterGlobals() override { FMeshPaintGeometryAdapterForStaticMeshes::InitializeAdapterGlobals(); }
-	virtual void AddReferencedObjectsGlobals(FReferenceCollector& Collector) override { FMeshPaintGeometryAdapterForStaticMeshes::AddReferencedObjectsGlobals(Collector); }
-	virtual void CleanupGlobals() override { FMeshPaintGeometryAdapterForStaticMeshes::CleanupGlobals(); }
+	virtual TSharedPtr<IMeshPaintComponentAdapter> Construct(class UMeshComponent* InComponent, int32 InMeshLODIndex) const override;
+	virtual void InitializeAdapterGlobals() override { FMeshPaintStaticMeshComponentAdapter::InitializeAdapterGlobals(); }
+	virtual void AddReferencedObjectsGlobals(FReferenceCollector& Collector) override { FMeshPaintStaticMeshComponentAdapter::AddReferencedObjectsGlobals(Collector); }
+	virtual void CleanupGlobals() override { FMeshPaintStaticMeshComponentAdapter::CleanupGlobals(); }
 };
diff --git a/Engine/Plugins/MeshPainting/Source/MeshPaintingToolset/Public/MeshVertexPaintingTool.h b/Engine/Plugins/MeshPainting/Source/MeshPaintingToolset/Public/MeshVertexPaintingTool.h
--- a/Engine/Plugins/MeshPainting/Source/MeshPaintingToolset/Public/MeshVertexPaintingTool.h
+++ b/Engine/Plugins/MeshPainting/Source/MeshPaintingToolset/Public/MeshVertexPaintingTool.h
@@ -16,7 +16,7 @@
 
 struct FToolBuilderState;
 enum class EMeshPaintModeAction : uint8;
-class IMeshPaintGeometryAdapter;
+class IMeshPaintComponentAdapter;
 
 struct FPaintRayResults
 {
diff --git a/Engine/Plugins/MeshPainting/Source/MeshPaintingToolset/Public/TexturePaintHelpers.h b/Engine/Plugins/MeshPainting/Source/MeshPaintingToolset/Public/TexturePaintHelpers.h
--- a/Engine/Plugins/MeshPainting/Source/MeshPaintingToolset/Public/TexturePaintHelpers.h
+++ b/Engine/Plugins/MeshPainting/Source/MeshPaintingToolset/Public/TexturePaintHelpers.h
@@ -52,7 +52,7 @@ struct FTexturePaintMeshSectionInfo
 	int32 LastIndex;
 };
 
-class IMeshPaintGeometryAdapter;
+class IMeshPaintComponentAdapter;
 struct FPaintableTexture;
 
 /** Helpers functions for texture painting functionality */
@@ -81,7 +81,7 @@ public:
 	static void RetrieveMeshSectionsForMaterialIndices(const UMeshComponent* MeshComponent, int32 LODIndex, const TArray<int32>& MaterialIndices, TArray<FTexturePaintMeshSectionInfo>& OutSectionInfo);
 
 	/** Retrieves all Paintable Textures from the given MeshComponent */
-	static void RetrieveTexturesForComponent(const UMeshComponent* Component, IMeshPaintGeometryAdapter* Adapter, TArray<FPaintableTexture>& OutTextures);
+	static void RetrieveTexturesForComponent(const UMeshComponent* Component, IMeshPaintComponentAdapter* Adapter, TArray<FPaintableTexture>& OutTextures);
 
 	/** Checks whether or not the given MeshComponent's materials use the specified texture */
 	static bool DoesMeshComponentUseTexture(UMeshComponent* MeshComponent, UTexture* Texture);
