Renaming adapter factories for consistency
#rnx

diff --git a/Engine/Plugins/MeshPainting/Source/MeshPaintEditorMode/Private/MeshPaintEditorMode.cpp b/Engine/Plugins/MeshPainting/Source/MeshPaintEditorMode/Private/MeshPaintEditorMode.cpp
--- a/Engine/Plugins/MeshPainting/Source/MeshPaintEditorMode/Private/MeshPaintEditorMode.cpp
+++ b/Engine/Plugins/MeshPainting/Source/MeshPaintEditorMode/Private/MeshPaintEditorMode.cpp
@@ -9,7 +9,7 @@
 #include "MeshPaintModeCommands.h"
 #include "BrushSettingsCustomization.h"
 #include "PropertyEditorModule.h"
-#include "IMeshPaintGeometryAdapterFactory.h"
+#include "MeshPaintComponentAdapterFactory.h"
 #include "MeshPaintAdapterFactory.h"
 #include "MeshPaintSplineMeshAdapter.h"
 #include "MeshPaintStaticMeshAdapter.h"
@@ -50,15 +50,15 @@ protected:
 
 	void OnMeshPaintModeButtonClicked();
 	bool IsMeshPaintModeButtonEnabled();
-	void RegisterGeometryAdapterFactory(TSharedRef<IMeshPaintGeometryAdapterFactory> Factory);
-	void UnregisterGeometryAdapterFactory(TSharedRef<IMeshPaintGeometryAdapterFactory> Factory);
+	void RegisterGeometryAdapterFactory(TSharedRef<IMeshPaintComponentAdapterFactory> Factory);
+	void UnregisterGeometryAdapterFactory(TSharedRef<IMeshPaintComponentAdapterFactory> Factory);
 };
 
 void FMeshPaintModeModule::Register()
 {
-	RegisterGeometryAdapterFactory(MakeShareable(new FMeshPaintGeometryAdapterForSplineMeshesFactory));
-	RegisterGeometryAdapterFactory(MakeShareable(new FMeshPaintGeometryAdapterForStaticMeshesFactory));
-	RegisterGeometryAdapterFactory(MakeShareable(new FMeshPaintGeometryAdapterForSkeletalMeshesFactory));
+	RegisterGeometryAdapterFactory(MakeShareable(new FMeshPaintSplineMeshComponentAdapterFactory));
+	RegisterGeometryAdapterFactory(MakeShareable(new FMeshPaintStaticMeshComponentAdapterFactory));
+	RegisterGeometryAdapterFactory(MakeShareable(new FMeshPaintSkeletalMeshComponentAdapterFactory));
 
 
 	FEditorModeRegistry::Get().RegisterScriptableMode<UMeshPaintMode>(
@@ -132,14 +132,14 @@ bool FMeshPaintModeModule::IsMeshPaintModeButtonEnabled( )
 	return true;
 }
 
-void  FMeshPaintModeModule::RegisterGeometryAdapterFactory(TSharedRef<IMeshPaintGeometryAdapterFactory> Factory)
+void  FMeshPaintModeModule::RegisterGeometryAdapterFactory(TSharedRef<IMeshPaintComponentAdapterFactory> Factory)
 {
-	FMeshPaintAdapterFactory::FactoryList.Add(Factory);
+	FMeshPaintComponentAdapterFactory::FactoryList.Add(Factory);
 }
 
-void  FMeshPaintModeModule::UnregisterGeometryAdapterFactory(TSharedRef<IMeshPaintGeometryAdapterFactory> Factory)
+void  FMeshPaintModeModule::UnregisterGeometryAdapterFactory(TSharedRef<IMeshPaintComponentAdapterFactory> Factory)
 {
-	FMeshPaintAdapterFactory::FactoryList.Remove(Factory);
+	FMeshPaintComponentAdapterFactory::FactoryList.Remove(Factory);
 }
 
 IMPLEMENT_MODULE(FMeshPaintModeModule, MeshPaintEditorMode)
diff --git a/Engine/Plugins/MeshPainting/Source/MeshPaintingToolset/Private/MeshPaintAdapterFactory.cpp b/Engine/Plugins/MeshPainting/Source/MeshPaintingToolset/Private/MeshPaintAdapterFactory.cpp
--- a/Engine/Plugins/MeshPainting/Source/MeshPaintingToolset/Private/MeshPaintAdapterFactory.cpp
+++ b/Engine/Plugins/MeshPainting/Source/MeshPaintingToolset/Private/MeshPaintAdapterFactory.cpp
@@ -1,11 +1,11 @@
 // Copyright 1998-2019 Epic Games, Inc. All Rights Reserved.
 
 #include "MeshPaintAdapterFactory.h"
-#include "IMeshPaintGeometryAdapterFactory.h"
+#include "MeshPaintComponentAdapterFactory.h"
 
-TArray<TSharedPtr<class IMeshPaintGeometryAdapterFactory>> FMeshPaintAdapterFactory::FactoryList;
+TArray<TSharedPtr<class IMeshPaintComponentAdapterFactory>> FMeshPaintComponentAdapterFactory::FactoryList;
 
-TSharedPtr<class IMeshPaintComponentAdapter> FMeshPaintAdapterFactory::CreateAdapterForMesh(UMeshComponent* InComponent, int32 InPaintingMeshLODIndex)
+TSharedPtr<class IMeshPaintComponentAdapter> FMeshPaintComponentAdapterFactory::CreateAdapterForMesh(UMeshComponent* InComponent, int32 InPaintingMeshLODIndex)
 {
 	TSharedPtr<IMeshPaintComponentAdapter> Result;
 
@@ -22,7 +22,7 @@ TSharedPtr<class IMeshPaintComponentAdapter> FMeshPaintAdapterFactory::CreateAda
 	return Result;
 }
 
-void FMeshPaintAdapterFactory::InitializeAdapterGlobals()
+void FMeshPaintComponentAdapterFactory::InitializeAdapterGlobals()
 {
 	for (const auto& Factory : FactoryList)
 	{
@@ -30,7 +30,7 @@ void FMeshPaintAdapterFactory::InitializeAdapterGlobals()
 	}
 }
 
-void FMeshPaintAdapterFactory::AddReferencedObjectsGlobals(FReferenceCollector& Collector)
+void FMeshPaintComponentAdapterFactory::AddReferencedObjectsGlobals(FReferenceCollector& Collector)
 {
 	for (const auto& Factory : FactoryList)
 	{
@@ -38,7 +38,7 @@ void FMeshPaintAdapterFactory::AddReferencedObjectsGlobals(FReferenceCollector&
 	}
 }
 
-void FMeshPaintAdapterFactory::CleanupGlobals()
+void FMeshPaintComponentAdapterFactory::CleanupGlobals()
 {
 	for (const auto& Factory : FactoryList)
 	{
diff --git a/Engine/Plugins/MeshPainting/Source/MeshPaintingToolset/Private/MeshPaintHelpers.cpp b/Engine/Plugins/MeshPainting/Source/MeshPaintingToolset/Private/MeshPaintHelpers.cpp
--- a/Engine/Plugins/MeshPainting/Source/MeshPaintingToolset/Private/MeshPaintHelpers.cpp
+++ b/Engine/Plugins/MeshPainting/Source/MeshPaintingToolset/Private/MeshPaintHelpers.cpp
@@ -1492,7 +1492,7 @@ void UMeshToolManager::CleanUp()
 		MeshAdapterPair.Value->OnRemoved();
 	}
 	ComponentToAdapterMap.Empty();
-	FMeshPaintAdapterFactory::CleanupGlobals();
+	FMeshPaintComponentAdapterFactory::CleanupGlobals();
 }
 
 bool UMeshToolManager::FindHitResult(const FRay Ray, FHitResult& BestTraceResult)
@@ -1553,7 +1553,7 @@ void UMeshToolManager::CacheSelectionData(const int32 PaintLODIndex, const int32
 	bSelectionContainsPerLODColors = false;
 	for (UMeshComponent* MeshComponent : SelectedMeshComponents)
 	{
-		TSharedPtr<IMeshPaintComponentAdapter> MeshAdapter = FMeshPaintAdapterFactory::CreateAdapterForMesh(MeshComponent, PaintLODIndex);
+		TSharedPtr<IMeshPaintComponentAdapter> MeshAdapter = FMeshPaintComponentAdapterFactory::CreateAdapterForMesh(MeshComponent, PaintLODIndex);
 		if (MeshComponent->IsVisible() && MeshAdapter.IsValid() && MeshAdapter->IsValid())
 		{
 			TUniquePtr< FComponentReregisterContext > ComponentReregisterContext;
diff --git a/Engine/Plugins/MeshPainting/Source/MeshPaintingToolset/Private/MeshPaintSkeletalMeshAdapter.cpp b/Engine/Plugins/MeshPainting/Source/MeshPaintingToolset/Private/MeshPaintSkeletalMeshAdapter.cpp
--- a/Engine/Plugins/MeshPainting/Source/MeshPaintingToolset/Private/MeshPaintSkeletalMeshAdapter.cpp
+++ b/Engine/Plugins/MeshPainting/Source/MeshPaintingToolset/Private/MeshPaintSkeletalMeshAdapter.cpp
@@ -313,7 +313,7 @@ FMatrix FMeshPaintSkeletalMeshComponentAdapter::GetComponentToWorldMatrix() cons
 //////////////////////////////////////////////////////////////////////////
 // FMeshPaintGeometryAdapterForSkeletalMeshesFactory
 
-TSharedPtr<IMeshPaintComponentAdapter> FMeshPaintGeometryAdapterForSkeletalMeshesFactory::Construct(class UMeshComponent* InComponent, int32 InMeshLODIndex) const
+TSharedPtr<IMeshPaintComponentAdapter> FMeshPaintSkeletalMeshComponentAdapterFactory::Construct(class UMeshComponent* InComponent, int32 InMeshLODIndex) const
 {
 	if (USkeletalMeshComponent* SkeletalMeshComponent = Cast<USkeletalMeshComponent>(InComponent))
 	{
diff --git a/Engine/Plugins/MeshPainting/Source/MeshPaintingToolset/Private/MeshPaintSplineMeshAdapter.cpp b/Engine/Plugins/MeshPainting/Source/MeshPaintingToolset/Private/MeshPaintSplineMeshAdapter.cpp
--- a/Engine/Plugins/MeshPainting/Source/MeshPaintingToolset/Private/MeshPaintSplineMeshAdapter.cpp
+++ b/Engine/Plugins/MeshPainting/Source/MeshPaintingToolset/Private/MeshPaintSplineMeshAdapter.cpp
@@ -50,7 +50,7 @@ bool FMeshPaintSplineMeshComponentAdapter::InitializeVertexData()
 //////////////////////////////////////////////////////////////////////////
 // FMeshPaintGeometryAdapterForSplineMeshesFactory
 
-TSharedPtr<IMeshPaintComponentAdapter> FMeshPaintGeometryAdapterForSplineMeshesFactory::Construct(class UMeshComponent* InComponent, int32 MeshLODIndex) const
+TSharedPtr<IMeshPaintComponentAdapter> FMeshPaintSplineMeshComponentAdapterFactory::Construct(class UMeshComponent* InComponent, int32 MeshLODIndex) const
 {
 	if (USplineMeshComponent* SplineMeshComponent = Cast<USplineMeshComponent>(InComponent))
 	{
diff --git a/Engine/Plugins/MeshPainting/Source/MeshPaintingToolset/Private/MeshPaintStaticMeshAdapter.cpp b/Engine/Plugins/MeshPainting/Source/MeshPaintingToolset/Private/MeshPaintStaticMeshAdapter.cpp
--- a/Engine/Plugins/MeshPainting/Source/MeshPaintingToolset/Private/MeshPaintStaticMeshAdapter.cpp
+++ b/Engine/Plugins/MeshPainting/Source/MeshPaintingToolset/Private/MeshPaintStaticMeshAdapter.cpp
@@ -426,7 +426,7 @@ void FMeshPaintStaticMeshComponentAdapter::PreEdit()
 //////////////////////////////////////////////////////////////////////////
 // FMeshPaintGeometryAdapterForStaticMeshesFactory
 
-TSharedPtr<IMeshPaintComponentAdapter> FMeshPaintGeometryAdapterForStaticMeshesFactory::Construct(class UMeshComponent* InComponent, int32 InMeshLODIndex) const
+TSharedPtr<IMeshPaintComponentAdapter> FMeshPaintStaticMeshComponentAdapterFactory::Construct(class UMeshComponent* InComponent, int32 InMeshLODIndex) const
 {
 	if (UStaticMeshComponent* StaticMeshComponent = Cast<UStaticMeshComponent>(InComponent))
 	{
diff --git a/Engine/Plugins/MeshPainting/Source/MeshPaintingToolset/Private/MeshSelect.cpp b/Engine/Plugins/MeshPainting/Source/MeshPaintingToolset/Private/MeshSelect.cpp
--- a/Engine/Plugins/MeshPainting/Source/MeshPaintingToolset/Private/MeshSelect.cpp
+++ b/Engine/Plugins/MeshPainting/Source/MeshPaintingToolset/Private/MeshSelect.cpp
@@ -83,7 +83,7 @@ void UMeshClickTool::OnClicked(const FInputDeviceRay& ClickPos)
 				for (UActorComponent* CandidateComponent : CandidateComponents)
 				{
 					UMeshComponent* MeshComponent = Cast<UMeshComponent>(CandidateComponent);
-					TSharedPtr<IMeshPaintComponentAdapter> MeshAdapter = FMeshPaintAdapterFactory::CreateAdapterForMesh(MeshComponent, 0);
+					TSharedPtr<IMeshPaintComponentAdapter> MeshAdapter = FMeshPaintComponentAdapterFactory::CreateAdapterForMesh(MeshComponent, 0);
 					MeshToolManager->AddPaintableMeshComponent(MeshComponent);
 					MeshToolManager->AddToComponentToAdapterMap(MeshComponent, MeshAdapter);
 				}
@@ -96,7 +96,7 @@ void UMeshClickTool::OnClicked(const FInputDeviceRay& ClickPos)
 		TArray<AActor*> SelectedActors;
 		if (UMeshComponent* MeshComponent = Cast<UMeshComponent>(TraceHitResult.GetComponent()))
 		{
-			TSharedPtr<IMeshPaintComponentAdapter> MeshAdapter = FMeshPaintAdapterFactory::CreateAdapterForMesh(MeshComponent, 0);
+			TSharedPtr<IMeshPaintComponentAdapter> MeshAdapter = FMeshPaintComponentAdapterFactory::CreateAdapterForMesh(MeshComponent, 0);
 			if (MeshComponent->IsVisible() && MeshAdapter.IsValid() && MeshAdapter->IsValid())
 			{
 				MeshToolManager->AddPaintableMeshComponent(MeshComponent);
diff --git a/Engine/Plugins/MeshPainting/Source/MeshPaintingToolset/Public/IMeshPaintGeometryAdapterFactory.h b/Engine/Plugins/MeshPainting/Source/MeshPaintingToolset/Public/IMeshPaintGeometryAdapterFactory.h
deleted file mode 100644
diff --git a/Engine/Plugins/MeshPainting/Source/MeshPaintingToolset/Public/MeshPaintAdapterFactory.h b/Engine/Plugins/MeshPainting/Source/MeshPaintingToolset/Public/MeshPaintAdapterFactory.h
--- a/Engine/Plugins/MeshPainting/Source/MeshPaintingToolset/Public/MeshPaintAdapterFactory.h
+++ b/Engine/Plugins/MeshPainting/Source/MeshPaintingToolset/Public/MeshPaintAdapterFactory.h
@@ -5,13 +5,13 @@
 #include "CoreMinimal.h"
 
 class UMeshComponent;
-class IMeshPaintGeometryAdapterFactory;
+class IMeshPaintComponentAdapterFactory;
 class FReferenceCollector;
 
-class MESHPAINTINGTOOLSET_API FMeshPaintAdapterFactory
+class MESHPAINTINGTOOLSET_API FMeshPaintComponentAdapterFactory
 {
 public:
-	static TArray<TSharedPtr<IMeshPaintGeometryAdapterFactory>> FactoryList;
+	static TArray<TSharedPtr<IMeshPaintComponentAdapterFactory>> FactoryList;
 
 public:
 	static TSharedPtr<class IMeshPaintComponentAdapter> CreateAdapterForMesh(UMeshComponent* InComponent, int32 InPaintingMeshLODIndex);
diff --git a/Engine/Plugins/MeshPainting/Source/MeshPaintingToolset/Public/MeshPaintComponentAdapterFactory.h b/Engine/Plugins/MeshPainting/Source/MeshPaintingToolset/Public/MeshPaintComponentAdapterFactory.h
new file mode 100644
--- /dev/null
+++ b/Engine/Plugins/MeshPainting/Source/MeshPaintingToolset/Public/MeshPaintComponentAdapterFactory.h
@@ -0,0 +1,20 @@
+// Copyright 1998-2019 Epic Games, Inc. All Rights Reserved.
+
+#pragma once
+
+class IMeshPaintComponentAdapter;
+class UMeshComponent;
+class FReferenceCollector;
+
+/**
+ * Factory for IMeshPaintGeometryAdapter
+ */
+class MESHPAINTINGTOOLSET_API IMeshPaintComponentAdapterFactory
+{
+public:
+	virtual TSharedPtr<IMeshPaintComponentAdapter> Construct(UMeshComponent* InComponent, int32 InMeshLODIndex) const = 0;
+	virtual void InitializeAdapterGlobals() = 0;
+	virtual void AddReferencedObjectsGlobals(FReferenceCollector& Collector) = 0;
+	virtual void CleanupGlobals() = 0;
+	virtual ~IMeshPaintComponentAdapterFactory() {}
+};
\ No newline at end of file
diff --git a/Engine/Plugins/MeshPainting/Source/MeshPaintingToolset/Public/MeshPaintSkeletalMeshAdapter.h b/Engine/Plugins/MeshPainting/Source/MeshPaintingToolset/Public/MeshPaintSkeletalMeshAdapter.h
--- a/Engine/Plugins/MeshPainting/Source/MeshPaintingToolset/Public/MeshPaintSkeletalMeshAdapter.h
+++ b/Engine/Plugins/MeshPainting/Source/MeshPaintingToolset/Public/MeshPaintSkeletalMeshAdapter.h
@@ -6,7 +6,7 @@
 #include "RawIndexBuffer.h"
 #include "Components/SkeletalMeshComponent.h"
 #include "BaseMeshPaintComponentAdapter.h"
-#include "IMeshPaintGeometryAdapterFactory.h"
+#include "MeshPaintComponentAdapterFactory.h"
 
 class UBodySetup;
 class USkeletalMesh;
@@ -76,7 +76,7 @@ protected:
 //////////////////////////////////////////////////////////////////////////
 // FMeshPaintGeometryAdapterForSkeletalMeshesFactory
 
-class MESHPAINTINGTOOLSET_API FMeshPaintGeometryAdapterForSkeletalMeshesFactory : public IMeshPaintGeometryAdapterFactory
+class MESHPAINTINGTOOLSET_API FMeshPaintSkeletalMeshComponentAdapterFactory : public IMeshPaintComponentAdapterFactory
 {
 public:
 	virtual TSharedPtr<IMeshPaintComponentAdapter> Construct(UMeshComponent* InComponent, int32 InMeshLODIndex) const override;
diff --git a/Engine/Plugins/MeshPainting/Source/MeshPaintingToolset/Public/MeshPaintSplineMeshAdapter.h b/Engine/Plugins/MeshPainting/Source/MeshPaintingToolset/Public/MeshPaintSplineMeshAdapter.h
--- a/Engine/Plugins/MeshPainting/Source/MeshPaintingToolset/Public/MeshPaintSplineMeshAdapter.h
+++ b/Engine/Plugins/MeshPainting/Source/MeshPaintingToolset/Public/MeshPaintSplineMeshAdapter.h
@@ -17,7 +17,7 @@ public:
 //////////////////////////////////////////////////////////////////////////
 // FMeshPaintGeometryAdapterForSplineMeshesFactory
 
-class MESHPAINTINGTOOLSET_API FMeshPaintGeometryAdapterForSplineMeshesFactory : public FMeshPaintGeometryAdapterForStaticMeshesFactory
+class MESHPAINTINGTOOLSET_API FMeshPaintSplineMeshComponentAdapterFactory : public FMeshPaintStaticMeshComponentAdapterFactory
 {
 public:
 	virtual TSharedPtr<IMeshPaintComponentAdapter> Construct(class UMeshComponent* InComponent, int32 MeshLODIndex) const override;
diff --git a/Engine/Plugins/MeshPainting/Source/MeshPaintingToolset/Public/MeshPaintStaticMeshAdapter.h b/Engine/Plugins/MeshPainting/Source/MeshPaintingToolset/Public/MeshPaintStaticMeshAdapter.h
--- a/Engine/Plugins/MeshPainting/Source/MeshPaintingToolset/Public/MeshPaintStaticMeshAdapter.h
+++ b/Engine/Plugins/MeshPainting/Source/MeshPaintingToolset/Public/MeshPaintStaticMeshAdapter.h
@@ -3,10 +3,10 @@
 #pragma once
 
 #include "CoreMinimal.h"
-#include "IMeshPaintGeometryAdapterFactory.h"
 #include "BaseMeshPaintComponentAdapter.h"
 #include "RawIndexBuffer.h"
 #include "Components/StaticMeshComponent.h"
+#include "MeshPaintComponentAdapterFactory.h"
 
 class UBodySetup;
 class UStaticMesh;
@@ -66,7 +66,7 @@ protected:
 //////////////////////////////////////////////////////////////////////////
 // FMeshPaintGeometryAdapterForStaticMeshesFactory
 
-class MESHPAINTINGTOOLSET_API FMeshPaintGeometryAdapterForStaticMeshesFactory : public IMeshPaintGeometryAdapterFactory
+class MESHPAINTINGTOOLSET_API FMeshPaintStaticMeshComponentAdapterFactory : public IMeshPaintComponentAdapterFactory
 {
 public:
 	virtual TSharedPtr<IMeshPaintComponentAdapter> Construct(class UMeshComponent* InComponent, int32 InMeshLODIndex) const override;
