Renaming mesh paint editor mode module for consistency

diff --git a/Engine/Plugins/MeshPainting/Source/MeshPaintEditorMode/Private/MeshPaintEditorMode.cpp b/Engine/Plugins/MeshPainting/Source/MeshPaintEditorMode/Private/MeshPaintEditorMode.cpp
--- a/Engine/Plugins/MeshPainting/Source/MeshPaintEditorMode/Private/MeshPaintEditorMode.cpp
+++ b/Engine/Plugins/MeshPainting/Source/MeshPaintEditorMode/Private/MeshPaintEditorMode.cpp
@@ -19,10 +19,10 @@
 
 #define LOCTEXT_NAMESPACE "MeshPaintMode"
 
-class FMeshPaintModeModule : public IModuleInterface
+class FMeshPaintEditorModeModule : public IModuleInterface
 {
 public:
-	FMeshPaintModeModule() 
+	FMeshPaintEditorModeModule() 
 	{
 	}
 
@@ -54,7 +54,7 @@ protected:
 	void UnregisterGeometryAdapterFactory(TSharedRef<IMeshPaintComponentAdapterFactory> Factory);
 };
 
-void FMeshPaintModeModule::Register()
+void FMeshPaintEditorModeModule::Register()
 {
 	RegisterGeometryAdapterFactory(MakeShareable(new FMeshPaintSplineMeshComponentAdapterFactory));
 	RegisterGeometryAdapterFactory(MakeShareable(new FMeshPaintStaticMeshComponentAdapterFactory));
@@ -75,7 +75,7 @@ void FMeshPaintModeModule::Register()
 }
 
 
-void FMeshPaintModeModule::OnPostEngineInit()
+void FMeshPaintEditorModeModule::OnPostEngineInit()
 {
 	if (!GetDefault<ULevelEditorMiscSettings>()->bEnableLegacyMeshPaintMode)
 	{
@@ -85,13 +85,13 @@ void FMeshPaintModeModule::OnPostEngineInit()
 	}
 }
 
-void FMeshPaintModeModule::StartupModule()
+void FMeshPaintEditorModeModule::StartupModule()
 {
-	FCoreDelegates::OnPostEngineInit.AddRaw(this, &FMeshPaintModeModule::OnPostEngineInit);
+	FCoreDelegates::OnPostEngineInit.AddRaw(this, &FMeshPaintEditorModeModule::OnPostEngineInit);
 }
 
 
-void FMeshPaintModeModule::ShutdownModule()
+void FMeshPaintEditorModeModule::ShutdownModule()
 {
 	
 	FCoreDelegates::OnPostEngineInit.RemoveAll(this);
@@ -108,13 +108,13 @@ void FMeshPaintModeModule::ShutdownModule()
 
 
 
-void FMeshPaintModeModule::Unregister()
+void FMeshPaintEditorModeModule::Unregister()
 {
 	FEditorModeRegistry::Get().UnregisterMode(GetEditorModeID());
 }
 
 
-void FMeshPaintModeModule::OnMeshPaintModeButtonClicked()
+void FMeshPaintEditorModeModule::OnMeshPaintModeButtonClicked()
 {
 	// *Important* - activate the mode first since FEditorModeTools::DeactivateMode will
 	// activate the default mode when the stack becomes empty, resulting in multiple active visible modes.
@@ -127,21 +127,21 @@ void FMeshPaintModeModule::OnMeshPaintModeButtonClicked()
 
 
 
-bool FMeshPaintModeModule::IsMeshPaintModeButtonEnabled( )
+bool FMeshPaintEditorModeModule::IsMeshPaintModeButtonEnabled( )
 {
 	return true;
 }
 
-void  FMeshPaintModeModule::RegisterGeometryAdapterFactory(TSharedRef<IMeshPaintComponentAdapterFactory> Factory)
+void  FMeshPaintEditorModeModule::RegisterGeometryAdapterFactory(TSharedRef<IMeshPaintComponentAdapterFactory> Factory)
 {
 	FMeshPaintComponentAdapterFactory::FactoryList.Add(Factory);
 }
 
-void  FMeshPaintModeModule::UnregisterGeometryAdapterFactory(TSharedRef<IMeshPaintComponentAdapterFactory> Factory)
+void  FMeshPaintEditorModeModule::UnregisterGeometryAdapterFactory(TSharedRef<IMeshPaintComponentAdapterFactory> Factory)
 {
 	FMeshPaintComponentAdapterFactory::FactoryList.Remove(Factory);
 }
 
-IMPLEMENT_MODULE(FMeshPaintModeModule, MeshPaintEditorMode)
+IMPLEMENT_MODULE(FMeshPaintEditorModeModule, MeshPaintEditorMode)
 
 #undef LOCTEXT_NAMESPACE
