Added recycled scriptable modes as a referenced object; Removed manual naming of tools to avoid name conflicts; Removed duplicate conditional that was preventing weight painting;

diff --git a/Engine/Plugins/MeshPainting/Source/MeshPaintEditorMode/Private/MeshPaintMode.cpp b/Engine/Plugins/MeshPainting/Source/MeshPaintEditorMode/Private/MeshPaintMode.cpp
--- a/Engine/Plugins/MeshPainting/Source/MeshPaintEditorMode/Private/MeshPaintMode.cpp
+++ b/Engine/Plugins/MeshPainting/Source/MeshPaintEditorMode/Private/MeshPaintMode.cpp
@@ -348,19 +348,19 @@ void UMeshPaintMode::FillWithVertexColor()
 			MeshAdapter->PreEdit();
 		}
 
-		if (UMeshColorPaintingToolProperties* ColorProperties = UMeshPaintMode::GetColorToolProperties())
-		{
-			const bool bPaintOnSpecificLOD = ColorProperties ? ColorProperties->bPaintOnSpecificLOD : false;
+		UMeshColorPaintingToolProperties* ColorProperties = UMeshPaintMode::GetColorToolProperties();
+		
+		const bool bPaintOnSpecificLOD = ColorProperties ? ColorProperties->bPaintOnSpecificLOD : false;
 
-			if (Component->IsA<UStaticMeshComponent>())
-			{
-				UMeshPaintingToolset::FillStaticMeshVertexColors(Cast<UStaticMeshComponent>(Component), bPaintOnSpecificLOD ? ColorProperties->LODIndex : -1, FillColor, MaskColor);
-			}
-			else if (Component->IsA<USkeletalMeshComponent>())
-			{
-				UMeshPaintingToolset::FillSkeletalMeshVertexColors(Cast<USkeletalMeshComponent>(Component), bPaintOnSpecificLOD ? ColorProperties->LODIndex : -1, FillColor, MaskColor);
-			}
+		if (Component->IsA<UStaticMeshComponent>())
+		{
+			UMeshPaintingToolset::FillStaticMeshVertexColors(Cast<UStaticMeshComponent>(Component), bPaintOnSpecificLOD ? ColorProperties->LODIndex : -1, FillColor, MaskColor);
+		}
+		else if (Component->IsA<USkeletalMeshComponent>())
+		{
+			UMeshPaintingToolset::FillSkeletalMeshVertexColors(Cast<USkeletalMeshComponent>(Component), bPaintOnSpecificLOD ? ColorProperties->LODIndex : -1, FillColor, MaskColor);
 		}
+		
 
 		if (MeshAdapter)
 		{
diff --git a/Engine/Plugins/MeshPainting/Source/MeshPaintingToolset/Private/MeshTexturePaintingTool.cpp b/Engine/Plugins/MeshPainting/Source/MeshPaintingToolset/Private/MeshTexturePaintingTool.cpp
--- a/Engine/Plugins/MeshPainting/Source/MeshPaintingToolset/Private/MeshTexturePaintingTool.cpp
+++ b/Engine/Plugins/MeshPainting/Source/MeshPaintingToolset/Private/MeshTexturePaintingTool.cpp
@@ -33,7 +33,7 @@ bool UMeshTexturePaintingToolBuilder::CanBuildTool(const FToolBuilderState& Scen
 
 UInteractiveTool* UMeshTexturePaintingToolBuilder::BuildTool(const FToolBuilderState& SceneState) const
 {
-	UMeshTexturePaintingTool* NewTool = NewObject<UMeshTexturePaintingTool>(SceneState.ToolManager, "ColorBrushTool");
+	UMeshTexturePaintingTool* NewTool = NewObject<UMeshTexturePaintingTool>(SceneState.ToolManager);
 	return NewTool;
 }
 
diff --git a/Engine/Plugins/MeshPainting/Source/MeshPaintingToolset/Private/MeshVertexPaintingTool.cpp b/Engine/Plugins/MeshPainting/Source/MeshPaintingToolset/Private/MeshVertexPaintingTool.cpp
--- a/Engine/Plugins/MeshPainting/Source/MeshPaintingToolset/Private/MeshVertexPaintingTool.cpp
+++ b/Engine/Plugins/MeshPainting/Source/MeshPaintingToolset/Private/MeshVertexPaintingTool.cpp
@@ -28,7 +28,7 @@ bool UMeshColorPaintingToolBuilder::CanBuildTool(const FToolBuilderState& SceneS
 
 UInteractiveTool* UMeshColorPaintingToolBuilder::BuildTool(const FToolBuilderState& SceneState) const
 {
-	UMeshColorPaintingTool* NewTool = NewObject<UMeshColorPaintingTool>(SceneState.ToolManager, "ColorBrushTool");
+	UMeshColorPaintingTool* NewTool = NewObject<UMeshColorPaintingTool>(SceneState.ToolManager);
 	return NewTool;
 }
 
@@ -40,7 +40,7 @@ bool UMeshWeightPaintingToolBuilder::CanBuildTool(const FToolBuilderState& Scene
 
 UInteractiveTool* UMeshWeightPaintingToolBuilder::BuildTool(const FToolBuilderState& SceneState) const
 {
-	UMeshWeightPaintingTool* NewTool = NewObject<UMeshWeightPaintingTool>(SceneState.ToolManager, "ColorWeightTool");
+	UMeshWeightPaintingTool* NewTool = NewObject<UMeshWeightPaintingTool>(SceneState.ToolManager);
 	return NewTool;
 }
 
diff --git a/Engine/Source/Editor/UnrealEd/Private/EditorModeManager.cpp b/Engine/Source/Editor/UnrealEd/Private/EditorModeManager.cpp
--- a/Engine/Source/Editor/UnrealEd/Private/EditorModeManager.cpp
+++ b/Engine/Source/Editor/UnrealEd/Private/EditorModeManager.cpp
@@ -2069,6 +2069,7 @@ void FEditorModeTools::AddReferencedObjects( FReferenceCollector& Collector )
 	}
 
 	Collector.AddReferencedObjects(ActiveScriptableModes);
+	Collector.AddReferencedObjects(RecycledScriptableModes);
 }
 
 FEdMode* FEditorModeTools::GetActiveMode( FEditorModeID InID )
