Fix deprecated warnings

diff --git a/Engine/Plugins/MeshPainting/Source/MeshPaintEditorMode/Private/MeshPaintModeHelpers.cpp b/Engine/Plugins/MeshPainting/Source/MeshPaintEditorMode/Private/MeshPaintModeHelpers.cpp
--- a/Engine/Plugins/MeshPainting/Source/MeshPaintEditorMode/Private/MeshPaintModeHelpers.cpp
+++ b/Engine/Plugins/MeshPainting/Source/MeshPaintEditorMode/Private/MeshPaintModeHelpers.cpp
@@ -123,20 +123,18 @@ void UMeshPaintModeHelpers::SetRealtimeViewport(bool bRealtime)
 {
 	FLevelEditorModule& LevelEditorModule = FModuleManager::GetModuleChecked<FLevelEditorModule>("LevelEditor");
 	TSharedPtr< IAssetViewport > ViewportWindow = LevelEditorModule.GetFirstActiveViewport();
-	const bool bRememberCurrentState = false;
 	if (ViewportWindow.IsValid())
 	{
-		FEditorViewportClient &Viewport = ViewportWindow->GetAssetViewportClient();
+		FEditorViewportClient& Viewport = ViewportWindow->GetAssetViewportClient();
 		if (Viewport.IsPerspective())
 		{
-			if (bRealtime)
+			if (!bRealtime)
 			{
-				Viewport.SetRealtime(bRealtime, bRememberCurrentState);
+				Viewport.SetRealtimeOverride(bRealtime, NSLOCTEXT("MeshPaint", "RealtimeOverrideMessage_MeshPaint", "Mesh Paint"));
 			}
 			else
 			{
-				const bool bAllowDisable = true;
-				Viewport.RestoreRealtime(bAllowDisable);
+				Viewport.RemoveRealtimeOverride();
 			}
 		}
 	}
