Fix deprecation warning 2

diff --git a/Engine/Plugins/MeshPainting/Source/MeshPaintEditorMode/Private/MeshPaintModeHelpers.cpp b/Engine/Plugins/MeshPainting/Source/MeshPaintEditorMode/Private/MeshPaintModeHelpers.cpp
--- a/Engine/Plugins/MeshPainting/Source/MeshPaintEditorMode/Private/MeshPaintModeHelpers.cpp
+++ b/Engine/Plugins/MeshPainting/Source/MeshPaintEditorMode/Private/MeshPaintModeHelpers.cpp
@@ -128,7 +128,7 @@ void UMeshPaintModeHelpers::SetRealtimeViewport(bool bRealtime)
 		FEditorViewportClient& Viewport = ViewportWindow->GetAssetViewportClient();
 		if (Viewport.IsPerspective())
 		{
-			if (!bRealtime)
+			if (bRealtime)
 			{
 				Viewport.SetRealtimeOverride(bRealtime, NSLOCTEXT("MeshPaint", "RealtimeOverrideMessage_MeshPaint", "Mesh Paint"));
 			}
