Moving brush settings up to the top of the various painting tool tabs

diff --git a/Engine/Plugins/MeshPainting/Source/MeshPaintEditorMode/Private/BrushSettingsCustomization.cpp b/Engine/Plugins/MeshPainting/Source/MeshPaintEditorMode/Private/BrushSettingsCustomization.cpp
--- a/Engine/Plugins/MeshPainting/Source/MeshPaintEditorMode/Private/BrushSettingsCustomization.cpp
+++ b/Engine/Plugins/MeshPainting/Source/MeshPaintEditorMode/Private/BrushSettingsCustomization.cpp
@@ -53,6 +53,8 @@ TSharedRef<IDetailCustomization> FVertexPaintingSettingsCustomization::MakeInsta
 
 void FVertexPaintingSettingsCustomization::CustomizeDetails(IDetailLayoutBuilder& DetailLayout)
 {	
+
+	IDetailCategoryBuilder& BrushCategory = DetailLayout.EditCategory(TEXT("Brush"), FText::GetEmpty(), ECategoryPriority::Important);
 	IDetailCategoryBuilder& VertexCategory = DetailLayout.EditCategory(TEXT("VertexPainting"));
 
 	VertexCategory.AddCustomRow(NSLOCTEXT("VertexPaintSettings", "InstanceColorSize", "Instance Color Size"))
@@ -368,4 +370,15 @@ void FWeightPaintingSettingsCustomization::OnTextureWeightTypeChanged(TSharedRef
 	EraseWeightProperty->SetValue(Value);
 }
 
-#undef LOCTEXT_NAMESPACE
\ No newline at end of file
+
+TSharedRef<IDetailCustomization> FTexturePaintingSettingsCustomization::MakeInstance()
+{
+	return MakeShareable(new FTexturePaintingSettingsCustomization);
+}
+
+void FTexturePaintingSettingsCustomization::CustomizeDetails(IDetailLayoutBuilder& DetailLayout)
+{
+	IDetailCategoryBuilder& BrushCategory = DetailLayout.EditCategory(TEXT("Brush"), FText::GetEmpty(), ECategoryPriority::Important);
+}
+
+#undef LOCTEXT_NAMESPACE
diff --git a/Engine/Plugins/MeshPainting/Source/MeshPaintEditorMode/Private/BrushSettingsCustomization.h b/Engine/Plugins/MeshPainting/Source/MeshPaintEditorMode/Private/BrushSettingsCustomization.h
--- a/Engine/Plugins/MeshPainting/Source/MeshPaintEditorMode/Private/BrushSettingsCustomization.h
+++ b/Engine/Plugins/MeshPainting/Source/MeshPaintEditorMode/Private/BrushSettingsCustomization.h
@@ -52,6 +52,14 @@ protected:
 
 };
 
+
+class FTexturePaintingSettingsCustomization : public IDetailCustomization
+{
+public:
+	static TSharedRef<IDetailCustomization> MakeInstance();
+	/** IPropertyTypeCustomization interface */
+	virtual void CustomizeDetails(IDetailLayoutBuilder& DetailLayout) override;
+};
 // class FTexturePaintSettingsCustomization : public IPropertyTypeCustomization
 // {
 // public:
diff --git a/Engine/Plugins/MeshPainting/Source/MeshPaintEditorMode/Private/MeshPaintEditorMode.cpp b/Engine/Plugins/MeshPainting/Source/MeshPaintEditorMode/Private/MeshPaintEditorMode.cpp
--- a/Engine/Plugins/MeshPainting/Source/MeshPaintEditorMode/Private/MeshPaintEditorMode.cpp
+++ b/Engine/Plugins/MeshPainting/Source/MeshPaintEditorMode/Private/MeshPaintEditorMode.cpp
@@ -72,6 +72,7 @@ void FMeshPaintEditorModeModule::Register()
 	FPropertyEditorModule& PropertyModule = FModuleManager::LoadModuleChecked<FPropertyEditorModule>("PropertyEditor");
 	PropertyModule.RegisterCustomClassLayout("MeshColorPaintingToolProperties", FOnGetDetailCustomizationInstance::CreateStatic(&FColorPaintingSettingsCustomization::MakeInstance));
 	PropertyModule.RegisterCustomClassLayout("MeshWeightPaintingToolProperties", FOnGetDetailCustomizationInstance::CreateStatic(&FWeightPaintingSettingsCustomization::MakeInstance));
+	PropertyModule.RegisterCustomClassLayout("MeshTexturePaintingToolProperties", FOnGetDetailCustomizationInstance::CreateStatic(&FTexturePaintingSettingsCustomization::MakeInstance));
 }
 
 
