Fixes for texture painting crashes
#rnx

diff --git a/Engine/Plugins/MeshPainting/Source/MeshPaintEditorMode/Private/BrushSettingsCustomization.cpp b/Engine/Plugins/MeshPainting/Source/MeshPaintEditorMode/Private/BrushSettingsCustomization.cpp
--- a/Engine/Plugins/MeshPainting/Source/MeshPaintEditorMode/Private/BrushSettingsCustomization.cpp
+++ b/Engine/Plugins/MeshPainting/Source/MeshPaintEditorMode/Private/BrushSettingsCustomization.cpp
@@ -26,6 +26,7 @@
 #include "MeshPaintMode.h"
 #include "Components/SkeletalMeshComponent.h"
 #include "MeshPaintModeHelpers.h"
+#include "MeshTexturePaintingTool.h"
 
 #define LOCTEXT_NAMESPACE "MeshPaintCustomization"
 
@@ -326,9 +327,6 @@ void FWeightPaintingSettingsCustomization::CustomizeDetails(IDetailLayoutBuilder
 {
 	FVertexPaintingSettingsCustomization::CustomizeDetails(DetailLayout);
 
-	TMap<FName, TSharedRef<IPropertyHandle>> CustomizedProperties;
-	TMap<FName, TSharedRef<IPropertyHandle>> Properties;
-
 	TSharedRef<IPropertyHandle> WeightTypeProperty = DetailLayout.GetProperty(GET_MEMBER_NAME_CHECKED(UMeshWeightPaintingToolProperties, TextureWeightType));
 	TSharedRef<IPropertyHandle> PaintWeightProperty = DetailLayout.GetProperty(GET_MEMBER_NAME_CHECKED(UMeshWeightPaintingToolProperties, PaintTextureWeightIndex));
 	TSharedRef<IPropertyHandle> EraseWeightProperty = DetailLayout.GetProperty(GET_MEMBER_NAME_CHECKED(UMeshWeightPaintingToolProperties, EraseTextureWeightIndex));
@@ -379,6 +377,41 @@ TSharedRef<IDetailCustomization> FTexturePaintingSettingsCustomization::MakeInst
 void FTexturePaintingSettingsCustomization::CustomizeDetails(IDetailLayoutBuilder& DetailLayout)
 {
 	IDetailCategoryBuilder& BrushCategory = DetailLayout.EditCategory(TEXT("Brush"), FText::GetEmpty(), ECategoryPriority::Important);
+	TSharedRef<IPropertyHandle> UVChannel = DetailLayout.GetProperty(GET_MEMBER_NAME_CHECKED(UMeshTexturePaintingToolProperties, UVChannel));
+	UVChannel->MarkHiddenByCustomization();
+	IDetailCategoryBuilder& ColorCategory = DetailLayout.EditCategory(TEXT("TexturePainting"));
+	ColorCategory.AddCustomRow(LOCTEXT("TexturePaintingUVLabel", "Texture Painting UV Channel"))
+		.NameContent()
+		[
+			UVChannel->CreatePropertyNameWidget()
+		]
+	.ValueContent()
+		[
+			SNew(SHorizontalBox)
+			+ SHorizontalBox::Slot()
+		.Padding(0.0f, 0.0f, 4.0f, 0.0f)
+		[
+			SNew(SNumericEntryBox<int32>)
+			.Font(IDetailLayoutBuilder::GetDetailFont())
+			.AllowSpin(true)
+			.Value_Lambda([]() -> int32 { return UMeshPaintMode::GetTextureToolProperties() ? UMeshPaintMode::GetTextureToolProperties()->UVChannel : 0; })
+			.MinValue(0)
+			.MaxValue_Lambda([]() -> int32 { return Cast<UMeshToolManager>(UMeshPaintMode::GetMeshPaintMode()->GetToolManager())->GetMaxUVIndexToPaint(); })
+			.OnValueChanged(SNumericEntryBox<int32>::FOnValueChanged::CreateLambda([=](int32 Value) { 
+				if (UMeshPaintMode::GetTextureToolProperties())
+				{
+					UMeshPaintMode::GetTextureToolProperties()->UVChannel = Value;
+				}
+				}))
+			.OnValueCommitted(SNumericEntryBox<int32>::FOnValueCommitted::CreateLambda([](int32 Value, ETextCommit::Type CommitType) {
+					if (UMeshPaintMode::GetTextureToolProperties())
+					{
+						UMeshPaintMode::GetTextureToolProperties()->UVChannel = Value;
+					}
+				}))
+			]
+		];
+
 }
 
 #undef LOCTEXT_NAMESPACE
diff --git a/Engine/Plugins/MeshPainting/Source/MeshPaintEditorMode/Private/BrushSettingsCustomization.h b/Engine/Plugins/MeshPainting/Source/MeshPaintEditorMode/Private/BrushSettingsCustomization.h
--- a/Engine/Plugins/MeshPainting/Source/MeshPaintEditorMode/Private/BrushSettingsCustomization.h
+++ b/Engine/Plugins/MeshPainting/Source/MeshPaintEditorMode/Private/BrushSettingsCustomization.h
@@ -60,20 +60,6 @@ public:
 	/** IPropertyTypeCustomization interface */
 	virtual void CustomizeDetails(IDetailLayoutBuilder& DetailLayout) override;
 };
-// class FTexturePaintSettingsCustomization : public IPropertyTypeCustomization
-// {
-// public:
-// 	static TSharedRef<IPropertyTypeCustomization> MakeInstance();
-// 
-// 	/** IPropertyTypeCustomization interface */
-// 	virtual void CustomizeHeader(TSharedRef<IPropertyHandle> PropertyHandle, FDetailWidgetRow& HeaderRow, IPropertyTypeCustomizationUtils& CustomizationUtils) override {};
-// 	virtual void CustomizeChildren(TSharedRef<IPropertyHandle> PropertyHandle, IDetailChildrenBuilder& ChildBuilder, IPropertyTypeCustomizationUtils& CustomizationUtils) override;
-// protected:
-// 	/** Cached instance of the painter */
-// 	FPaintModePainter* MeshPainter;
-// 	/** Cached instance of the texture painting settings */
-// 	FTexturePaintSettings* PaintSettings;
-// };
 
 
 	
\ No newline at end of file
diff --git a/Engine/Plugins/MeshPainting/Source/MeshPaintEditorMode/Private/MeshPaintMode.cpp b/Engine/Plugins/MeshPainting/Source/MeshPaintEditorMode/Private/MeshPaintMode.cpp
--- a/Engine/Plugins/MeshPainting/Source/MeshPaintEditorMode/Private/MeshPaintMode.cpp
+++ b/Engine/Plugins/MeshPainting/Source/MeshPaintEditorMode/Private/MeshPaintMode.cpp
@@ -699,6 +699,7 @@ void UMeshPaintMode::UpdateCachedVertexDataSize()
 			}
 		}
 	}
+	bRecacheVertexDataSize = false;
 }
 
 
diff --git a/Engine/Plugins/MeshPainting/Source/MeshPaintingToolset/Private/MeshPaintHelpers.cpp b/Engine/Plugins/MeshPainting/Source/MeshPaintingToolset/Private/MeshPaintHelpers.cpp
--- a/Engine/Plugins/MeshPainting/Source/MeshPaintingToolset/Private/MeshPaintHelpers.cpp
+++ b/Engine/Plugins/MeshPainting/Source/MeshPaintingToolset/Private/MeshPaintHelpers.cpp
@@ -1574,6 +1574,16 @@ void UMeshToolManager::CacheSelectionData(const int32 PaintLODIndex, const int32
 	bNeedsRecache = false;
 }
 
+int32 UMeshToolManager::GetMaxUVIndexToPaint() const
+{
+	if (PaintableComponents.Num() == 1 && PaintableComponents[0])
+	{
+		return UMeshPaintingToolset::GetNumberOfUVs(PaintableComponents[0], 0) - 1;
+	}
+
+	return 0;
+}
+
 UMeshToolsContext::UMeshToolsContext()
 	:Super()
 {
diff --git a/Engine/Plugins/MeshPainting/Source/MeshPaintingToolset/Private/MeshPaintSkeletalMeshAdapter.cpp b/Engine/Plugins/MeshPainting/Source/MeshPaintingToolset/Private/MeshPaintSkeletalMeshAdapter.cpp
--- a/Engine/Plugins/MeshPainting/Source/MeshPaintingToolset/Private/MeshPaintSkeletalMeshAdapter.cpp
+++ b/Engine/Plugins/MeshPainting/Source/MeshPaintingToolset/Private/MeshPaintSkeletalMeshAdapter.cpp
@@ -43,12 +43,15 @@ FMeshPaintSkeletalMeshComponentAdapter::~FMeshPaintSkeletalMeshComponentAdapter(
 void FMeshPaintSkeletalMeshComponentAdapter::OnSkeletalMeshChanged()
 {
 	OnRemoved();
-	ReferencedSkeletalMesh = SkeletalMeshComponent->SkeletalMesh;
-	if (SkeletalMeshComponent->SkeletalMesh != nullptr)
-	{	
-		Initialize();
-		OnAdded();
-	}	
+	if (SkeletalMeshComponent != nullptr)
+	{
+		ReferencedSkeletalMesh = SkeletalMeshComponent->SkeletalMesh;
+		if (SkeletalMeshComponent->SkeletalMesh != nullptr)
+		{
+			Initialize();
+			OnAdded();
+		}
+	}
 }
 
 void FMeshPaintSkeletalMeshComponentAdapter::OnPostMeshCached(USkeletalMesh* SkeletalMesh)
diff --git a/Engine/Plugins/MeshPainting/Source/MeshPaintingToolset/Private/MeshTexturePaintingTool.cpp b/Engine/Plugins/MeshPainting/Source/MeshPaintingToolset/Private/MeshTexturePaintingTool.cpp
--- a/Engine/Plugins/MeshPainting/Source/MeshPaintingToolset/Private/MeshTexturePaintingTool.cpp
+++ b/Engine/Plugins/MeshPainting/Source/MeshPaintingToolset/Private/MeshTexturePaintingTool.cpp
@@ -61,21 +61,21 @@ UMeshTexturePaintingToolProperties::UMeshTexturePaintingToolProperties()
 void UMeshTexturePaintingToolProperties::SaveProperties(UInteractiveTool* SaveFromTool)
 {
 	UBrushBaseProperties::SaveProperties(SaveFromTool);
-// 	UMeshVertexPaintingToolProperties* PropertyCache = GetPropertyCache<UMeshVertexPaintingToolProperties>();
-// 	PropertyCache->PaintColor = this->PaintColor;
-// 	PropertyCache->EraseColor = this->EraseColor;
-// 	PropertyCache->bEnableFlow = this->bEnableFlow;
-// 	PropertyCache->bOnlyFrontFacingTriangles = this->bOnlyFrontFacingTriangles;
+	UMeshTexturePaintingToolProperties* PropertyCache = GetPropertyCache<UMeshTexturePaintingToolProperties>();
+	PropertyCache->PaintColor = this->PaintColor;
+	PropertyCache->EraseColor = this->EraseColor;
+	PropertyCache->bEnableFlow = this->bEnableFlow;
+	PropertyCache->bOnlyFrontFacingTriangles = this->bOnlyFrontFacingTriangles;
 }
 
 void UMeshTexturePaintingToolProperties::RestoreProperties(UInteractiveTool* RestoreToTool)
 {
 	UBrushBaseProperties::RestoreProperties(RestoreToTool);
-// 	UMeshVertexPaintingToolProperties* PropertyCache = GetPropertyCache<UMeshVertexPaintingToolProperties>();
-// 	this->PaintColor = PropertyCache->PaintColor;
-// 	this->EraseColor = PropertyCache->EraseColor;
-// 	this->bEnableFlow = PropertyCache->bEnableFlow;
-// 	this->bOnlyFrontFacingTriangles = PropertyCache->bOnlyFrontFacingTriangles;
+	UMeshVertexPaintingToolProperties* PropertyCache = GetPropertyCache<UMeshVertexPaintingToolProperties>();
+	this->PaintColor = PropertyCache->PaintColor;
+	this->EraseColor = PropertyCache->EraseColor;
+	this->bEnableFlow = PropertyCache->bEnableFlow;
+	this->bOnlyFrontFacingTriangles = PropertyCache->bOnlyFrontFacingTriangles;
 }
 
 UMeshTexturePaintingTool::UMeshTexturePaintingTool()
@@ -122,7 +122,6 @@ void UMeshTexturePaintingTool::Render(IToolsContextRenderAPI* RenderAPI)
 		static const FLinearColor NormalLineColor(0.3f, 1.0f, 0.3f);
 		const FLinearColor BrushCueColor = (bArePainting ? FLinearColor(1.0f, 1.0f, 0.3f) : FLinearColor(0.3f, 1.0f, 0.3f));
  		const FLinearColor InnerBrushCueColor = (bArePainting ? FLinearColor(0.5f, 0.5f, 0.1f) : FLinearColor(0.1f, 0.5f, 0.1f));
-		//const float PointDrawSize = VertexProperties->VertexPreviewSize;
 		// Draw trace surface normal
 		const FVector NormalLineEnd(LastBestHitResult.Location + LastBestHitResult.Normal * NormalLineSize);
 		Draw.DrawLine(FVector(LastBestHitResult.Location), NormalLineEnd, NormalLineColor, WidgetLineThickness);
@@ -131,7 +130,7 @@ void UMeshTexturePaintingTool::Render(IToolsContextRenderAPI* RenderAPI)
 		{
 			TSharedPtr<IMeshPaintComponentAdapter> MeshAdapter = MeshToolManager->GetAdapterForComponent(Cast<UMeshComponent>(CurrentComponent));
 
-			if (MeshAdapter->IsValid() /*&& bRenderVertices*/ && MeshAdapter->SupportsVertexPaint())
+			if (MeshAdapter->IsValid() && MeshAdapter->SupportsVertexPaint())
 			{
 				const FMatrix ComponentToWorldMatrix = MeshAdapter->GetComponentToWorldMatrix();
 				FViewCameraState CameraState;
@@ -142,19 +141,6 @@ void UMeshTexturePaintingTool::Render(IToolsContextRenderAPI* RenderAPI)
 				// @todo MeshPaint: Input vector doesn't work well with non-uniform scale
 				const float ComponentSpaceBrushRadius = ComponentToWorldMatrix.InverseTransformVector(FVector(BrushProperties->BrushRadius, 0.0f, 0.0f)).Size();
 				const float ComponentSpaceSquaredBrushRadius = ComponentSpaceBrushRadius * ComponentSpaceBrushRadius;
-
-			//	const TArray<FVector>& InRangeVertices = MeshAdapter->SphereIntersectVertices(ComponentSpaceSquaredBrushRadius, ComponentSpaceBrushPosition, ComponentSpaceCameraPosition, VertexProperties->bOnlyFrontFacingTriangles);
-
-// 				for (const FVector& Vertex : InRangeVertices)
-// 				{
-// 					const FVector WorldPositionVertex = ComponentToWorldMatrix.TransformPosition(Vertex);
-// 					if ((LastBestHitResult.Location - WorldPositionVertex).Size() <= BrushProperties->BrushRadius)
-// 					{
-// 						const float VisualBiasDistance = 0.15f;
-// 						const FVector VertexVisualPosition = WorldPositionVertex + LastBestHitResult.Normal * VisualBiasDistance;
-// 						Draw.DrawPoint(VertexVisualPosition, HoverVertexPointColor, PointDrawSize, SDPG_World);
-// 					}
-// 				}
 			}
 		}
 		Draw.EndFrame();
@@ -183,10 +169,10 @@ void UMeshTexturePaintingTool::Tick(float DeltaTime)
 		bStampPending = false;
 
 		// flow
-// 		if (bInDrag && VertexProperties && VertexProperties->bEnableFlow)
-// 		{
-// 			bStampPending = true;
-// 		}
+		if (bInDrag && TextureProperties && TextureProperties->bEnableFlow)
+		{
+			bStampPending = true;
+		}
 	}
 }
 
@@ -231,7 +217,7 @@ void UMeshTexturePaintingTool::CacheSelectionData()
 		MeshToolManager->ClearPaintableMeshComponents();
 
 		//Determine LOD level to use for painting(can only paint on LODs in vertex mode)
-		const int32 PaintLODIndex = 0;//ColorProperties->bPaintOnSpecificLOD ? ColorProperties->LODIndex : 0;
+		const int32 PaintLODIndex = 0;
 		//Determine UV channel to use while painting textures
 		const int32 UVChannel = 0;
 
@@ -386,7 +372,7 @@ bool UMeshTexturePaintingTool::PaintInternal(const TArrayView<TPair<FVector, FVe
 
 			if (MeshAdapter->SupportsTexturePaint())
 			{
-				TArray<const UTexture*> Textures;
+				Textures.Empty();
 				const UTexture2D* TargetTexture2D = TextureProperties->PaintTexture;
 				if (TargetTexture2D)
 				{
diff --git a/Engine/Plugins/MeshPainting/Source/MeshPaintingToolset/Public/MeshPaintHelpers.h b/Engine/Plugins/MeshPainting/Source/MeshPaintingToolset/Public/MeshPaintHelpers.h
--- a/Engine/Plugins/MeshPainting/Source/MeshPaintingToolset/Public/MeshPaintHelpers.h
+++ b/Engine/Plugins/MeshPainting/Source/MeshPaintingToolset/Public/MeshPaintHelpers.h
@@ -288,6 +288,7 @@ public:
 	TArray<FPerComponentVertexColorData> GetCopiedColorsByComponent() const;
 	void SetCopiedColorsByComponent(TArray<FPerComponentVertexColorData>& InCopiedColors);
 	void CacheSelectionData(const int32 PaintLODIndex, const int32 UVChannel);
+	int32 GetMaxUVIndexToPaint() const;
 	void ResetState();
 	void Refresh();
 	bool SelectionContainsPerLODColors() const { return bSelectionContainsPerLODColors; }
diff --git a/Engine/Plugins/MeshPainting/Source/MeshPaintingToolset/Public/MeshPaintingToolsetTypes.h b/Engine/Plugins/MeshPainting/Source/MeshPaintingToolset/Public/MeshPaintingToolsetTypes.h
--- a/Engine/Plugins/MeshPainting/Source/MeshPaintingToolset/Public/MeshPaintingToolsetTypes.h
+++ b/Engine/Plugins/MeshPainting/Source/MeshPaintingToolset/Public/MeshPaintingToolsetTypes.h
@@ -109,22 +109,29 @@ struct FPaintableTexture
 	}
 };
 
+USTRUCT()
 struct FPaintTexture2DData
 {
+	GENERATED_BODY()
 	/** The original texture that we're painting */
+	UPROPERTY(Transient)
 	UTexture2D* PaintingTexture2D;
 	bool bIsPaintingTexture2DModified;
 
 	/** A copy of the original texture we're painting, used for restoration. */
+	UPROPERTY(Transient)
 	UTexture2D* PaintingTexture2DDuplicate;
 
 	/** Render target texture for painting */
+	UPROPERTY(Transient)
 	UTextureRenderTarget2D* PaintRenderTargetTexture;
 
 	/** Render target texture used as an input while painting that contains a clone of the original image */
+	UPROPERTY(Transient)
 	UTextureRenderTarget2D* CloneRenderTargetTexture;
 
 	/** List of materials we are painting on */
+	UPROPERTY(Transient)
 	TArray< UMaterialInterface* > PaintingMaterials;
 
 	/** Default ctor */
diff --git a/Engine/Plugins/MeshPainting/Source/MeshPaintingToolset/Public/MeshTexturePaintingTool.h b/Engine/Plugins/MeshPainting/Source/MeshPaintingToolset/Public/MeshTexturePaintingTool.h
--- a/Engine/Plugins/MeshPainting/Source/MeshPaintingToolset/Public/MeshTexturePaintingTool.h
+++ b/Engine/Plugins/MeshPainting/Source/MeshPaintingToolset/Public/MeshTexturePaintingTool.h
@@ -46,27 +46,27 @@ class MESHPAINTINGTOOLSET_API UMeshTexturePaintingToolProperties : public UBrush
 public:
 	UMeshTexturePaintingToolProperties();
 
-	/** Color used for Applying Vertex Color Painting */
-	UPROPERTY(EditAnywhere, Category = VertexPainting)
+	/** Color used for Applying Texture Color Painting */
+	UPROPERTY(EditAnywhere, Category = TexturePainting)
 	FLinearColor PaintColor;
 
-	/** Color used for Erasing Vertex Color Painting */
-	UPROPERTY(EditAnywhere, Category = VertexPainting)
+	/** Color used for Erasing Texture Color Painting */
+	UPROPERTY(EditAnywhere, Category = TexturePainting)
 	FLinearColor EraseColor;
 
-	/** Whether or not to apply Vertex Color Painting to the Red Channel */
+	/** Whether or not to apply Texture Color Painting to the Red Channel */
 	UPROPERTY(EditAnywhere, Category = TexturePainting, meta = (DisplayName = "Red"))
 	bool bWriteRed;
 
-	/** Whether or not to apply Vertex Color Painting to the Green Channel */
+	/** Whether or not to apply Texture Color Painting to the Green Channel */
 	UPROPERTY(EditAnywhere, Category = TexturePainting, meta = (DisplayName = "Green"))
 	bool bWriteGreen;
 
-	/** Whether or not to apply Vertex Color Painting to the Blue Channel */
+	/** Whether or not to apply Texture Color Painting to the Blue Channel */
 	UPROPERTY(EditAnywhere, Category = TexturePainting, meta = (DisplayName = "Blue"))
 	bool bWriteBlue;
 
-	/** Whether or not to apply Vertex Color Painting to the Alpha Channel */
+	/** Whether or not to apply Texture Color Painting to the Alpha Channel */
 	UPROPERTY(EditAnywhere, Category = TexturePainting, meta = (DisplayName = "Alpha"))
 	bool bWriteAlpha;
 
@@ -164,6 +164,10 @@ private:
 private:
 	UPROPERTY(Transient)
 	UMeshTexturePaintingToolProperties* TextureProperties;
+
+	UPROPERTY(Transient)
+	TArray<const UTexture*> Textures;
+
 	/** Flag for whether or not we are currently painting */
 	bool bArePainting;
 	bool bDoRestoreRenTargets;
@@ -180,24 +184,30 @@ private:
 	TMap<UMeshComponent*, FInstanceTexturePaintSettings> ComponentToTexturePaintSettingsMap;
 
 	/** Temporary render target used to draw incremental paint to */
+	UPROPERTY(Transient)
 	UTextureRenderTarget2D* BrushRenderTargetTexture;
 
 	/** Temporary render target used to store a mask of the affected paint region, updated every time we add incremental texture paint */
+	UPROPERTY(Transient)
 	UTextureRenderTarget2D* BrushMaskRenderTargetTexture;
 
 	/** Temporary render target used to store generated mask for texture seams, we create this by projecting object triangles into texture space using the selected UV channel */
+	UPROPERTY(Transient)
 	UTextureRenderTarget2D* SeamMaskRenderTargetTexture;
 
 	/** Stores data associated with our paint target textures */
+	UPROPERTY(Transient)
 	TMap< UTexture2D*, FPaintTexture2DData > PaintTargetData;
 
 	/** Texture paint: Will hold a list of texture items that we can paint on */
 	TArray<FTextureTargetListInfo> TexturePaintTargetList;
 
 	/** Texture paint: The mesh components that we're currently painting */
+	UPROPERTY(Transient)
 	UMeshComponent* TexturePaintingCurrentMeshComponent;
 
 	/** The original texture that we're painting */
+	UPROPERTY(Transient)
 	UTexture2D* PaintingTexture2D;
 
 	/** True if we need to generate a texture seam mask used for texture dilation */
