Adding texture paint details customizations; fixing brush indicator drawing while no paintable mesh is present.
#rnx

diff --git a/Engine/Plugins/MeshPainting/Source/MeshPaintEditorMode/Private/BrushSettingsCustomization.cpp b/Engine/Plugins/MeshPainting/Source/MeshPaintEditorMode/Private/BrushSettingsCustomization.cpp
--- a/Engine/Plugins/MeshPainting/Source/MeshPaintEditorMode/Private/BrushSettingsCustomization.cpp
+++ b/Engine/Plugins/MeshPainting/Source/MeshPaintEditorMode/Private/BrushSettingsCustomization.cpp
@@ -27,6 +27,7 @@
 #include "Components/SkeletalMeshComponent.h"
 #include "MeshPaintModeHelpers.h"
 #include "MeshTexturePaintingTool.h"
+#include "ContentBrowserDelegates.h"
 
 #define LOCTEXT_NAMESPACE "MeshPaintCustomization"
 
@@ -161,7 +162,7 @@ FReply FVertexPaintingSettingsCustomization::OnSwapColorsClicked(TSharedRef<IPro
 {
 	FScopedTransaction Transaction(NSLOCTEXT("VertexPaintSettings", "SwapColorsTransation", "Swap paint and erase colors"));
 
-	UMeshPaintModeHelpers::SwapColors();
+	UMeshPaintModeHelpers::SwapVertexColors();
 	UMeshVertexPaintingToolProperties* Settings = UMeshPaintMode::GetVertexToolProperties();
 	if (Settings)
 	{
@@ -377,10 +378,48 @@ TSharedRef<IDetailCustomization> FTexturePaintingSettingsCustomization::MakeInst
 void FTexturePaintingSettingsCustomization::CustomizeDetails(IDetailLayoutBuilder& DetailLayout)
 {
 	IDetailCategoryBuilder& BrushCategory = DetailLayout.EditCategory(TEXT("Brush"), FText::GetEmpty(), ECategoryPriority::Important);
+	IDetailCategoryBuilder& TextureCategory = DetailLayout.EditCategory(TEXT("TexturePainting"));
+
+	TSharedRef<IPropertyHandle> RedChannel = DetailLayout.GetProperty("bWriteRed", UMeshTexturePaintingToolProperties::StaticClass());
+	RedChannel->MarkHiddenByCustomization();
+	TSharedRef<IPropertyHandle> GreenChannel = DetailLayout.GetProperty("bWriteGreen", UMeshTexturePaintingToolProperties::StaticClass());
+	GreenChannel->MarkHiddenByCustomization();
+	TSharedRef<IPropertyHandle> BlueChannel = DetailLayout.GetProperty("bWriteBlue", UMeshTexturePaintingToolProperties::StaticClass());
+	BlueChannel->MarkHiddenByCustomization();
+	TSharedRef<IPropertyHandle> AlphaChannel = DetailLayout.GetProperty("bWriteAlpha", UMeshTexturePaintingToolProperties::StaticClass());
+	AlphaChannel->MarkHiddenByCustomization();
+	TArray<TSharedRef<IPropertyHandle>> Channels = { RedChannel, GreenChannel, BlueChannel, AlphaChannel };
+	TSharedPtr<SHorizontalBox> ChannelsWidget;
+
+	TextureCategory.AddCustomRow(NSLOCTEXT("VertexPaintSettings", "ChannelLabel", "Channels"))
+		.NameContent()
+		[
+			SNew(STextBlock)
+			.Text(NSLOCTEXT("VertexPaintSettings", "ChannelsLabel", "Channels"))
+			.ToolTipText(NSLOCTEXT("VertexPaintSettings", "ChannelsToolTip", "Colors Channels which should be influenced during Painting."))
+			.Font(IDetailLayoutBuilder::GetDetailFont())
+		]
+		.ValueContent()
+		.MaxDesiredWidth(250.0f)
+		[
+			SAssignNew(ChannelsWidget, SHorizontalBox)
+		];
+
+	for (TSharedRef<IPropertyHandle> Channel : Channels)
+	{
+		ChannelsWidget->AddSlot()
+			.AutoWidth()
+			.Padding(0.0f, 0.0f, 4.0f, 0.0f)
+			[
+				CreateColorChannelWidget(Channel)
+			];
+	}
+
+
 	TSharedRef<IPropertyHandle> UVChannel = DetailLayout.GetProperty(GET_MEMBER_NAME_CHECKED(UMeshTexturePaintingToolProperties, UVChannel));
 	UVChannel->MarkHiddenByCustomization();
-	IDetailCategoryBuilder& ColorCategory = DetailLayout.EditCategory(TEXT("TexturePainting"));
-	ColorCategory.AddCustomRow(LOCTEXT("TexturePaintingUVLabel", "Texture Painting UV Channel"))
+
+	TextureCategory.AddCustomRow(LOCTEXT("TexturePaintingUVLabel", "Texture Painting UV Channel"))
 		.NameContent()
 		[
 			UVChannel->CreatePropertyNameWidget()
@@ -412,6 +451,117 @@ void FTexturePaintingSettingsCustomization::CustomizeDetails(IDetailLayoutBuilde
 			]
 		];
 
+
+	TSharedRef<IPropertyHandle> TextureProperty = DetailLayout.GetProperty(GET_MEMBER_NAME_CHECKED(UMeshTexturePaintingToolProperties, PaintTexture));
+	TextureProperty->MarkHiddenByCustomization();
+	TSharedPtr<SHorizontalBox> TextureWidget;
+	FDetailWidgetRow& Row = TextureCategory.AddCustomRow(NSLOCTEXT("TexturePaintSetting", "TextureSearchString", "Texture"))
+		.NameContent()
+		[
+			SNew(STextBlock)
+			.Text(NSLOCTEXT("TexturePaintSettings", "PaintTextureLabel", "Paint Texture"))
+			.ToolTipText(NSLOCTEXT("TexturePaintSettings", "PaintTextureToolTip", "Texture to Apply Painting to."))
+			.Font(IDetailLayoutBuilder::GetDetailFont())
+		]
+		.ValueContent()
+		.MaxDesiredWidth(250.0f)
+		[
+			SAssignNew(TextureWidget, SHorizontalBox)
+		];
+
+	/** Use a SObjectPropertyEntryBox to benefit from its functionality */
+		TextureWidget->AddSlot()
+		[
+			SNew(SObjectPropertyEntryBox)
+			.PropertyHandle(TextureProperty)
+			.AllowedClass(UTexture2D::StaticClass())
+			.OnShouldFilterAsset(FOnShouldFilterAsset::CreateUObject(Cast<UMeshTexturePaintingTool>(UMeshPaintMode::GetMeshPaintMode()->GetToolManager()->GetActiveTool(EToolSide::Left)), &UMeshTexturePaintingTool::ShouldFilterTextureAsset))
+			.OnObjectChanged(FOnSetObject::CreateUObject(Cast<UMeshTexturePaintingTool>(UMeshPaintMode::GetMeshPaintMode()->GetToolManager()->GetActiveTool(EToolSide::Left)), &UMeshTexturePaintingTool::PaintTextureChanged))
+			.DisplayUseSelected(false)
+			.ThumbnailPool(DetailLayout.GetThumbnailPool())
+		];
+
+		/** Creates a custom widget row containing all color channel flags */
+		TSharedRef<IPropertyHandle> PaintColor = DetailLayout.GetProperty("PaintColor", UMeshTexturePaintingToolProperties::StaticClass());
+		PaintColor->MarkHiddenByCustomization();
+		TSharedRef<IPropertyHandle> EraseColor = DetailLayout.GetProperty("EraseColor", UMeshTexturePaintingToolProperties::StaticClass());
+		EraseColor->MarkHiddenByCustomization();
+
+		// Customize paint color with a swap button
+		{
+			TSharedPtr<SWidget> NameWidget;
+			TSharedPtr<SWidget> ValueWidget;
+
+			IDetailPropertyRow& PaintColorProp = TextureCategory.AddProperty(PaintColor);
+			PaintColorProp.GetDefaultWidgets(NameWidget, ValueWidget, false);
+			FDetailWidgetRow& PaintRow = PaintColorProp.CustomWidget(true);
+			PaintRow.NameContent()
+				[
+					NameWidget.ToSharedRef()
+				];
+
+			PaintRow.ValueContent()
+				.MinDesiredWidth(250)
+				.MaxDesiredWidth(0)
+				[
+					SNew(SHorizontalBox)
+					+ SHorizontalBox::Slot()
+					.FillWidth(1.0)
+					.HAlign(HAlign_Left)
+					[
+						SNew(SBox)
+						.WidthOverride(250.f)
+						[
+							ValueWidget.ToSharedRef()
+						]
+					]
+					+ SHorizontalBox::Slot()
+					.HAlign(HAlign_Center)
+					.AutoWidth()
+					[
+						SNew(SButton)
+						.ButtonStyle(FEditorStyle::Get(), "HoverHintOnly")
+						.ToolTipText(NSLOCTEXT("VertexPaintSettings", "SwapColors", "Swap Paint and Erase Colors"))
+						.HAlign(HAlign_Center)
+						.VAlign(VAlign_Center)
+						.OnClicked(this, &FTexturePaintingSettingsCustomization::OnSwapColorsClicked, PaintColor, EraseColor)
+						.ContentPadding(0)
+						[
+							SNew(SImage).Image(FEditorStyle::GetBrush("MeshPaint.Swap"))
+						]
+					]
+				];
+		}
+
+		{
+			IDetailPropertyRow& EraseColorProp = TextureCategory.AddProperty(EraseColor);
+
+			TSharedPtr<SWidget> NameWidget;
+			TSharedPtr<SWidget> ValueWidget;
+
+			FDetailWidgetRow& EraseRow = EraseColorProp.CustomWidget(true);
+			EraseRow.ValueContent().MinDesiredWidth(250 - 16.f);
+			EraseColorProp.GetDefaultWidgets(NameWidget, ValueWidget, EraseRow, false);
+		}
+
 }
 
+FReply FTexturePaintingSettingsCustomization::OnSwapColorsClicked(TSharedRef<IPropertyHandle> PaintColor, TSharedRef<IPropertyHandle> EraseColor)
+{
+	UMeshTexturePaintingToolProperties* Settings = UMeshPaintMode::GetTextureToolProperties();
+	if (Settings)
+	{
+		FScopedTransaction Transaction(LOCTEXT("SwapColorsTransation", "Swap paint and erase colors"));
+		Settings->Modify();
+
+		FLinearColor TempPaintColor = Settings->PaintColor;
+		Settings->PaintColor = Settings->EraseColor;
+		Settings->EraseColor = TempPaintColor;
+		PaintColor->NotifyPostChange(EPropertyChangeType::ValueSet);
+		EraseColor->NotifyPostChange(EPropertyChangeType::ValueSet);
+	}
+	return FReply::Handled();
+}
+
+
 #undef LOCTEXT_NAMESPACE
diff --git a/Engine/Plugins/MeshPainting/Source/MeshPaintEditorMode/Private/BrushSettingsCustomization.h b/Engine/Plugins/MeshPainting/Source/MeshPaintEditorMode/Private/BrushSettingsCustomization.h
--- a/Engine/Plugins/MeshPainting/Source/MeshPaintEditorMode/Private/BrushSettingsCustomization.h
+++ b/Engine/Plugins/MeshPainting/Source/MeshPaintEditorMode/Private/BrushSettingsCustomization.h
@@ -59,6 +59,7 @@ public:
 	static TSharedRef<IDetailCustomization> MakeInstance();
 	/** IPropertyTypeCustomization interface */
 	virtual void CustomizeDetails(IDetailLayoutBuilder& DetailLayout) override;
+	FReply OnSwapColorsClicked(TSharedRef<IPropertyHandle> PaintColor, TSharedRef<IPropertyHandle> EraseColor);
 };
 
 
diff --git a/Engine/Plugins/MeshPainting/Source/MeshPaintEditorMode/Private/MeshPaintMode.cpp b/Engine/Plugins/MeshPainting/Source/MeshPaintEditorMode/Private/MeshPaintMode.cpp
--- a/Engine/Plugins/MeshPainting/Source/MeshPaintEditorMode/Private/MeshPaintMode.cpp
+++ b/Engine/Plugins/MeshPainting/Source/MeshPaintEditorMode/Private/MeshPaintMode.cpp
@@ -203,7 +203,7 @@ void UMeshPaintMode::BindCommands()
 
 	CommandList->MapAction(Commands.SwitchForeAndBackgroundColor, FExecuteAction::CreateLambda([this]()
 	{
-		UMeshPaintModeHelpers::SwapColors();
+		UMeshPaintModeHelpers::SwapVertexColors();
 	}));
 
 
@@ -292,10 +292,7 @@ void UMeshPaintMode::OnToolEnded(UInteractiveToolManager* Manager, UInteractiveT
 
 bool UMeshPaintMode::InputAxis(FEditorViewportClient* InViewportClient, FViewport* Viewport, int32 ControllerId, FKey Key, float Delta, float DeltaTime)
 {
-	bool bSelectingOnly = (GetToolManager()->GetActiveTool(EToolSide::Mouse)) ? GetToolManager()->GetActiveTool(EToolSide::Mouse)->IsA<USingleClickTool>() : false;
-
-	// If we are painting (not selecting), don't allow panning with the mouse.
-	return !bSelectingOnly && !InViewportClient->IsAltPressed();
+	return !InViewportClient->IsAltPressed();
 }
 
 void UMeshPaintMode::ActorSelectionChangeNotify()
diff --git a/Engine/Plugins/MeshPainting/Source/MeshPaintEditorMode/Private/MeshPaintModeHelpers.cpp b/Engine/Plugins/MeshPainting/Source/MeshPaintEditorMode/Private/MeshPaintModeHelpers.cpp
--- a/Engine/Plugins/MeshPainting/Source/MeshPaintEditorMode/Private/MeshPaintModeHelpers.cpp
+++ b/Engine/Plugins/MeshPainting/Source/MeshPaintEditorMode/Private/MeshPaintModeHelpers.cpp
@@ -765,7 +765,7 @@ void UMeshPaintModeHelpers::RemovePerLODColors(const TArray<UMeshComponent*>& Pa
 	}
 }
 
-void UMeshPaintModeHelpers::SwapColors()
+void UMeshPaintModeHelpers::SwapVertexColors()
 {
 	UMeshVertexPaintingToolProperties* Settings = UMeshPaintMode::GetVertexToolProperties();
 	if (Settings)
diff --git a/Engine/Plugins/MeshPainting/Source/MeshPaintEditorMode/Public/MeshPaintModeHelpers.h b/Engine/Plugins/MeshPainting/Source/MeshPaintEditorMode/Public/MeshPaintModeHelpers.h
--- a/Engine/Plugins/MeshPainting/Source/MeshPaintEditorMode/Public/MeshPaintModeHelpers.h
+++ b/Engine/Plugins/MeshPainting/Source/MeshPaintEditorMode/Public/MeshPaintModeHelpers.h
@@ -70,7 +70,7 @@ public:
 	static bool CanPasteInstanceVertexColors(const TArray<UStaticMeshComponent*>& StaticMeshComponents, const TArray<FPerComponentVertexColorData>& CopiedColorsByComponent);
 	static void RemovePerLODColors(const TArray<UMeshComponent*>& PaintableComponents);
 
-	static void SwapColors();
+	static void SwapVertexColors();
 	static 	void SaveModifiedTextures();
 	static bool CanSaveModifiedTextures();
 };
\ No newline at end of file
diff --git a/Engine/Plugins/MeshPainting/Source/MeshPaintingToolset/Private/MeshTexturePaintingTool.cpp b/Engine/Plugins/MeshPainting/Source/MeshPaintingToolset/Private/MeshTexturePaintingTool.cpp
--- a/Engine/Plugins/MeshPainting/Source/MeshPaintingToolset/Private/MeshTexturePaintingTool.cpp
+++ b/Engine/Plugins/MeshPainting/Source/MeshPaintingToolset/Private/MeshTexturePaintingTool.cpp
@@ -16,6 +16,7 @@
 #include "MeshPaintingToolsetTypes.h"
 #include "CanvasItem.h"
 #include "ScopedTransaction.h"
+#include "BaseGizmos/BrushStampIndicator.h"
 
 #define LOCTEXT_NAMESPACE "MeshTextureBrush"
 
@@ -65,17 +66,25 @@ void UMeshTexturePaintingToolProperties::SaveProperties(UInteractiveTool* SaveFr
 	PropertyCache->PaintColor = this->PaintColor;
 	PropertyCache->EraseColor = this->EraseColor;
 	PropertyCache->bEnableFlow = this->bEnableFlow;
+	PropertyCache->bWriteRed = this->bWriteRed;
+	PropertyCache->bWriteGreen = this->bWriteGreen;
+	PropertyCache->bWriteBlue = this->bWriteBlue;
+	PropertyCache->bWriteRed = this->bWriteRed;
 	PropertyCache->bOnlyFrontFacingTriangles = this->bOnlyFrontFacingTriangles;
 }
 
 void UMeshTexturePaintingToolProperties::RestoreProperties(UInteractiveTool* RestoreToTool)
 {
 	UBrushBaseProperties::RestoreProperties(RestoreToTool);
-	UMeshVertexPaintingToolProperties* PropertyCache = GetPropertyCache<UMeshVertexPaintingToolProperties>();
+	UMeshTexturePaintingToolProperties* PropertyCache = GetPropertyCache<UMeshTexturePaintingToolProperties>();
 	this->PaintColor = PropertyCache->PaintColor;
 	this->EraseColor = PropertyCache->EraseColor;
 	this->bEnableFlow = PropertyCache->bEnableFlow;
 	this->bOnlyFrontFacingTriangles = PropertyCache->bOnlyFrontFacingTriangles;
+	this->bWriteRed = PropertyCache->bWriteRed;
+	this->bWriteGreen = PropertyCache->bWriteGreen;
+	this->bWriteBlue = PropertyCache->bWriteBlue;
+	this->bWriteRed = PropertyCache->bWriteRed;
 }
 
 UMeshTexturePaintingTool::UMeshTexturePaintingTool()
@@ -110,11 +119,12 @@ void UMeshTexturePaintingTool::Shutdown(EToolShutdownType ShutdownType)
 void UMeshTexturePaintingTool::Render(IToolsContextRenderAPI* RenderAPI)
 {
 	UMeshToolManager* MeshToolManager = Cast<UMeshToolManager>(GetToolManager());
+	Super::Render(RenderAPI);
+	FToolDataVisualizer Draw;
+	Draw.BeginFrame(RenderAPI);
 	if (MeshToolManager && LastBestHitResult.Component != nullptr)
 	{
-		Super::Render(RenderAPI);
-		FToolDataVisualizer Draw;
-		Draw.BeginFrame(RenderAPI);
+		BrushStampIndicator->bDrawIndicatorLines = true;
 		static float WidgetLineThickness = 1.0f;
 		static FLinearColor VertexPointColor = FLinearColor::White;
 		static FLinearColor	HoverVertexPointColor = FLinearColor(0.3f, 1.0f, 0.3f);
@@ -143,8 +153,12 @@ void UMeshTexturePaintingTool::Render(IToolsContextRenderAPI* RenderAPI)
 				const float ComponentSpaceSquaredBrushRadius = ComponentSpaceBrushRadius * ComponentSpaceBrushRadius;
 			}
 		}
-		Draw.EndFrame();
 	}
+	else
+	{
+		BrushStampIndicator->bDrawIndicatorLines = false;
+	}
+	Draw.EndFrame();
 	UpdateResult();
 
 }
@@ -1229,5 +1243,31 @@ int32 UMeshTexturePaintingTool::GetNumberOfPendingPaintChanges() const
 	return Result;
 }
 
+bool UMeshTexturePaintingTool::ShouldFilterTextureAsset(const FAssetData& AssetData) const
+{
+	return !(PaintableTextures.ContainsByPredicate([=](const FPaintableTexture& Texture) { return Texture.Texture->GetFullName() == AssetData.GetFullName(); }));
+}
+
+void UMeshTexturePaintingTool::PaintTextureChanged(const FAssetData& AssetData)
+{
+	UTexture2D* Texture = Cast<UTexture2D>(AssetData.GetAsset());
+	if (Texture)
+	{
+		// Loop through our list of textures and see which one the user wants to select
+		for (int32 TargetIndex = 0; TargetIndex < TexturePaintTargetList.Num(); TargetIndex++)
+		{
+			FTextureTargetListInfo& TextureTarget = TexturePaintTargetList[TargetIndex];
+			if (TextureTarget.TextureData == Texture)
+			{
+				TextureTarget.bIsSelected = true;
+				TextureProperties->UVChannel = TextureTarget.UVChannelIndex;
+			}
+			else
+			{
+				TextureTarget.bIsSelected = false;
+			}
+		}
+	}
+}
 
 #undef LOCTEXT_NAMESPACE
\ No newline at end of file
diff --git a/Engine/Plugins/MeshPainting/Source/MeshPaintingToolset/Private/MeshVertexPaintingTool.cpp b/Engine/Plugins/MeshPainting/Source/MeshPaintingToolset/Private/MeshVertexPaintingTool.cpp
--- a/Engine/Plugins/MeshPainting/Source/MeshPaintingToolset/Private/MeshVertexPaintingTool.cpp
+++ b/Engine/Plugins/MeshPainting/Source/MeshPaintingToolset/Private/MeshVertexPaintingTool.cpp
@@ -12,6 +12,7 @@
 #include "MeshPaintAdapterFactory.h"
 #include "ToolDataVisualizer.h"
 #include "MeshPaintHelpers.h"
+#include "BaseGizmos/BrushStampIndicator.h"
 
 
 #define LOCTEXT_NAMESPACE "MeshVertexBrush"
@@ -106,11 +107,12 @@ void UMeshVertexPaintingTool::Shutdown(EToolShutdownType ShutdownType)
 void UMeshVertexPaintingTool::Render(IToolsContextRenderAPI* RenderAPI)
 {
 	UMeshToolManager* MeshToolManager = Cast<UMeshToolManager>(GetToolManager());
-	if (MeshToolManager && LastBestHitResult.Component != nullptr)
+	Super::Render(RenderAPI);
+	FToolDataVisualizer Draw;
+	Draw.BeginFrame(RenderAPI);
+	if (MeshToolManager && LastBestHitResult.Component != nullptr && MeshToolManager->GetPaintableMeshComponents().Num() > 0)
 	{
-		Super::Render(RenderAPI);
-		FToolDataVisualizer Draw;
-		Draw.BeginFrame(RenderAPI);
+		BrushStampIndicator->bDrawIndicatorLines = true;
 		static float WidgetLineThickness = 1.0f;
 		static FLinearColor VertexPointColor = FLinearColor::White;
 		static FLinearColor	HoverVertexPointColor = FLinearColor(0.3f, 1.0f, 0.3f);
@@ -153,8 +155,12 @@ void UMeshVertexPaintingTool::Render(IToolsContextRenderAPI* RenderAPI)
 				}
 			}
 		}
-		Draw.EndFrame();
 	}
+	else
+	{
+		BrushStampIndicator->bDrawIndicatorLines = false;
+	}
+	Draw.EndFrame();
 	UpdateResult();
 
 }
@@ -672,11 +678,20 @@ UMeshWeightPaintingToolProperties::UMeshWeightPaintingToolProperties()
 void UMeshWeightPaintingToolProperties::SaveProperties(UInteractiveTool* SaveFromTool)
 {
 	Super::SaveProperties(SaveFromTool);
+	UMeshWeightPaintingToolProperties* PropertyCache = GetPropertyCache<UMeshWeightPaintingToolProperties>();
+	PropertyCache->TextureWeightType = this->TextureWeightType;
+	PropertyCache->PaintTextureWeightIndex = this->PaintTextureWeightIndex;
+	PropertyCache->EraseTextureWeightIndex = this->EraseTextureWeightIndex;
 }
 
 void UMeshWeightPaintingToolProperties::RestoreProperties(UInteractiveTool* RestoreToTool)
 {
 	Super::RestoreProperties(RestoreToTool);
+	UMeshWeightPaintingToolProperties* PropertyCache = GetPropertyCache<UMeshWeightPaintingToolProperties>();
+	this->TextureWeightType = PropertyCache->TextureWeightType;
+	this->PaintTextureWeightIndex = PropertyCache->PaintTextureWeightIndex;
+	this->EraseTextureWeightIndex = PropertyCache->EraseTextureWeightIndex;
+
 }
 
 UMeshWeightPaintingTool::UMeshWeightPaintingTool()
diff --git a/Engine/Plugins/MeshPainting/Source/MeshPaintingToolset/Public/MeshTexturePaintingTool.h b/Engine/Plugins/MeshPainting/Source/MeshPaintingToolset/Public/MeshTexturePaintingTool.h
--- a/Engine/Plugins/MeshPainting/Source/MeshPaintingToolset/Public/MeshTexturePaintingTool.h
+++ b/Engine/Plugins/MeshPainting/Source/MeshPaintingToolset/Public/MeshTexturePaintingTool.h
@@ -135,6 +135,8 @@ public:
 	/** Returns the number of texture that require a commit. */
 	int32 GetNumberOfPendingPaintChanges() const;
 
+	bool ShouldFilterTextureAsset(const FAssetData& AssetData) const;
+	void PaintTextureChanged(const FAssetData& AssetData);
 protected:
 	virtual void SetAdditionalPaintParameters(FMeshPaintParameters& InPaintParameters) {};
 	virtual void FinishPainting();
