Mesh Texture Painting no longer allows painting on source textures with more bytes per pixel than the created render target data.

diff --git a/Engine/Plugins/MeshPainting/Source/MeshPaintingToolset/Private/IMeshPaintComponentAdapter.cpp b/Engine/Plugins/MeshPainting/Source/MeshPaintingToolset/Private/IMeshPaintComponentAdapter.cpp
--- a/Engine/Plugins/MeshPainting/Source/MeshPaintingToolset/Private/IMeshPaintComponentAdapter.cpp
+++ b/Engine/Plugins/MeshPainting/Source/MeshPaintingToolset/Private/IMeshPaintComponentAdapter.cpp
@@ -10,6 +10,7 @@
 #include "Materials/Material.h"
 #include "Materials/MaterialExpressionTextureCoordinate.h"
 #include "Materials/MaterialExpressionTextureSampleParameter.h"
+#include "TexturePaintToolset.h"
 
 //////////////////////////////////////////////////////////////////////////
 // IMeshPaintGeometryAdapter
@@ -56,7 +57,8 @@ void IMeshPaintComponentAdapter::DefaultQueryPaintableTextures(int32 MaterialInd
 				TextureBase->Texture != NULL &&
 				!TextureBase->Texture->IsNormalMap() && 
 				!TextureBase->Texture->VirtualTextureStreaming &&
-				!TextureBase->Texture->HasHDRSource()) // Currently HDR textures are not supported to paint on
+				!TextureBase->Texture->HasHDRSource() && // Currently HDR textures are not supported to paint on.
+				(TextureBase->Texture->Source.GetBytesPerPixel() <= UTexturePaintToolset::GetMaxSupportedBytesPerPixelForPainting())) // Textures' sources must fit in FColor struct to be supported.
 			{
 				// Default UV channel to index 0. 
 				PaintableTexture = FPaintableTexture(TextureBase->Texture, 0);
diff --git a/Engine/Plugins/MeshPainting/Source/MeshPaintingToolset/Private/TexturePaintToolset.cpp b/Engine/Plugins/MeshPainting/Source/MeshPaintingToolset/Private/TexturePaintToolset.cpp
--- a/Engine/Plugins/MeshPainting/Source/MeshPaintingToolset/Private/TexturePaintToolset.cpp
+++ b/Engine/Plugins/MeshPainting/Source/MeshPaintingToolset/Private/TexturePaintToolset.cpp
@@ -337,7 +337,7 @@ UTexture2D* UTexturePaintToolset::CreateTempUncompressedTexture(UTexture2D* Sour
 	check(Width > 0 && Height > 0 && RawData.Num() > 0);
 
 	// Allocate the new texture
-	UTexture2D* NewTexture2D = UTexture2D::CreateTransient(Width, Height, PF_B8G8R8A8);
+	UTexture2D* NewTexture2D = UTexture2D::CreateTransient(Width, Height, GetTempUncompressedTexturePixelFormat());
 
 	// Fill in the base mip for the texture we created
 	uint8* MipData = (uint8*)NewTexture2D->PlatformData->Mips[0].BulkData.Lock(LOCK_READ_WRITE);
diff --git a/Engine/Plugins/MeshPainting/Source/MeshPaintingToolset/Public/TexturePaintToolset.h b/Engine/Plugins/MeshPainting/Source/MeshPaintingToolset/Public/TexturePaintToolset.h
--- a/Engine/Plugins/MeshPainting/Source/MeshPaintingToolset/Public/TexturePaintToolset.h
+++ b/Engine/Plugins/MeshPainting/Source/MeshPaintingToolset/Public/TexturePaintToolset.h
@@ -69,6 +69,18 @@ public:
 	/** Will generate a mask texture, used for texture dilation, and store it in the passed in render target */
 	static bool GenerateSeamMask(UMeshComponent* MeshComponent, int32 UVSet, UTextureRenderTarget2D* SeamRenderTexture, UTexture2D* Texture, UTextureRenderTarget2D* RenderTargetTexture );
 
+	/** Returns the maximum bytes per pixel that are supported for source textures when painting. This limitation is set by CreateTempUncompressedTexture. */
+	static constexpr int32 GetMaxSupportedBytesPerPixelForPainting()
+	{
+		return GPixelFormats[GetTempUncompressedTexturePixelFormat()].BlockBytes;
+	}
+
+	/** Returns the pixel format that CreateTempUncompressedTexture uses to create render target data for painting. */
+	static constexpr EPixelFormat GetTempUncompressedTexturePixelFormat()
+	{
+		return EPixelFormat::PF_B8G8R8A8;
+	}
+
 	/** Static: Creates a temporary texture used to transfer data to a render target in memory */
 	static UTexture2D* CreateTempUncompressedTexture(UTexture2D* SourceTexture);
 
diff --git a/Engine/Source/Editor/MeshPaint/Private/IMeshPaintGeometryAdapter.cpp b/Engine/Source/Editor/MeshPaint/Private/IMeshPaintGeometryAdapter.cpp
--- a/Engine/Source/Editor/MeshPaint/Private/IMeshPaintGeometryAdapter.cpp
+++ b/Engine/Source/Editor/MeshPaint/Private/IMeshPaintGeometryAdapter.cpp
@@ -10,6 +10,7 @@
 #include "Materials/Material.h"
 #include "Materials/MaterialExpressionTextureCoordinate.h"
 #include "Materials/MaterialExpressionTextureSampleParameter.h"
+#include "TexturePaintHelpers.h"
 
 //////////////////////////////////////////////////////////////////////////
 // IMeshPaintGeometryAdapter
@@ -56,7 +57,8 @@ void IMeshPaintGeometryAdapter::DefaultQueryPaintableTextures(int32 MaterialInde
 				TextureBase->Texture != NULL &&
 				!TextureBase->Texture->IsNormalMap() && 
 				!TextureBase->Texture->VirtualTextureStreaming &&
-				!TextureBase->Texture->HasHDRSource()) // Currently HDR textures are not supported to paint on 
+				!TextureBase->Texture->HasHDRSource() && // Currently HDR textures are not supported to paint on.
+				(TextureBase->Texture->Source.GetBytesPerPixel() <= TexturePaintHelpers::GetMaxSupportedBytesPerPixelForPainting())) // Textures' sources must fit in FColor struct to be supported.
 			{
 				// Default UV channel to index 0. 
 				PaintableTexture = FPaintableTexture(TextureBase->Texture, 0);
diff --git a/Engine/Source/Editor/MeshPaint/Private/TexturePaintHelpers.cpp b/Engine/Source/Editor/MeshPaint/Private/TexturePaintHelpers.cpp
--- a/Engine/Source/Editor/MeshPaint/Private/TexturePaintHelpers.cpp
+++ b/Engine/Source/Editor/MeshPaint/Private/TexturePaintHelpers.cpp
@@ -337,7 +337,7 @@ UTexture2D* TexturePaintHelpers::CreateTempUncompressedTexture(UTexture2D* Sourc
 	check(Width > 0 && Height > 0 && RawData.Num() > 0);
 
 	// Allocate the new texture
-	UTexture2D* NewTexture2D = UTexture2D::CreateTransient(Width, Height, PF_B8G8R8A8);
+	UTexture2D* NewTexture2D = UTexture2D::CreateTransient(Width, Height, GetTempUncompressedTexturePixelFormat());
 
 	// Fill in the base mip for the texture we created
 	uint8* MipData = (uint8*)NewTexture2D->PlatformData->Mips[0].BulkData.Lock(LOCK_READ_WRITE);
diff --git a/Engine/Source/Editor/MeshPaint/Public/TexturePaintHelpers.h b/Engine/Source/Editor/MeshPaint/Public/TexturePaintHelpers.h
--- a/Engine/Source/Editor/MeshPaint/Public/TexturePaintHelpers.h
+++ b/Engine/Source/Editor/MeshPaint/Public/TexturePaintHelpers.h
@@ -65,6 +65,18 @@ public:
 	/** Will generate a mask texture, used for texture dilation, and store it in the passed in render target */
 	static bool GenerateSeamMask(UMeshComponent* MeshComponent, int32 UVSet, UTextureRenderTarget2D* SeamRenderTexture, UTexture2D* Texture, UTextureRenderTarget2D* RenderTargetTexture );
 
+	/** Returns the maximum bytes per pixel that are supported for source textures when painting. This limitation is set by CreateTempUncompressedTexture. */
+	static constexpr int32 GetMaxSupportedBytesPerPixelForPainting()
+	{
+		return GPixelFormats[GetTempUncompressedTexturePixelFormat()].BlockBytes;
+	}
+
+	/** Returns the pixel format that CreateTempUncompressedTexture uses to create render target data for painting. */
+	static constexpr EPixelFormat GetTempUncompressedTexturePixelFormat()
+	{
+		return EPixelFormat::PF_B8G8R8A8;
+	}
+
 	/** Static: Creates a temporary texture used to transfer data to a render target in memory */
 	static UTexture2D* CreateTempUncompressedTexture(UTexture2D* SourceTexture);
 
