Fix for crash on undoing painting on skeletal mesh

diff --git a/Engine/Plugins/MeshPainting/Source/MeshPaintingToolset/Private/MeshPaintHelpers.cpp b/Engine/Plugins/MeshPainting/Source/MeshPaintingToolset/Private/MeshPaintHelpers.cpp
--- a/Engine/Plugins/MeshPainting/Source/MeshPaintingToolset/Private/MeshPaintHelpers.cpp
+++ b/Engine/Plugins/MeshPainting/Source/MeshPaintingToolset/Private/MeshPaintHelpers.cpp
@@ -1513,7 +1513,7 @@ bool UMeshToolManager::FindHitResult(const FRay Ray, FHitResult& BestTraceResult
 
 		for (UMeshComponent* MeshComponent : PaintableComponents)
 		{
-			TSharedPtr<IMeshPaintComponentAdapter> MeshAdapter = GetComponentToAdapterMap().FindChecked(MeshComponent);
+			TSharedPtr<IMeshPaintComponentAdapter> MeshAdapter = GetAdapterForComponent(MeshComponent);
 
 			// Ray trace
 			FHitResult TraceHitResult(1.0f);
diff --git a/Engine/Plugins/MeshPainting/Source/MeshPaintingToolset/Private/MeshPaintSkeletalMeshAdapter.cpp b/Engine/Plugins/MeshPainting/Source/MeshPaintingToolset/Private/MeshPaintSkeletalMeshAdapter.cpp
--- a/Engine/Plugins/MeshPainting/Source/MeshPaintingToolset/Private/MeshPaintSkeletalMeshAdapter.cpp
+++ b/Engine/Plugins/MeshPainting/Source/MeshPaintingToolset/Private/MeshPaintSkeletalMeshAdapter.cpp
@@ -38,6 +38,10 @@ FMeshPaintSkeletalMeshComponentAdapter::~FMeshPaintSkeletalMeshComponentAdapter(
 	{
 		SkeletalMeshComponent->UnregisterOnSkeletalMeshPropertyChanged(SkeletalMeshChangedHandle);
 	}
+	if (ReferencedSkeletalMesh != nullptr)
+	{
+		ReferencedSkeletalMesh->OnPostMeshCached().RemoveAll(this);
+	}
 }
 
 void FMeshPaintSkeletalMeshComponentAdapter::OnSkeletalMeshChanged()
diff --git a/Engine/Plugins/MeshPainting/Source/MeshPaintingToolset/Public/BaseMeshPaintComponentAdapter.h b/Engine/Plugins/MeshPainting/Source/MeshPaintingToolset/Public/BaseMeshPaintComponentAdapter.h
--- a/Engine/Plugins/MeshPainting/Source/MeshPaintingToolset/Public/BaseMeshPaintComponentAdapter.h
+++ b/Engine/Plugins/MeshPainting/Source/MeshPaintingToolset/Public/BaseMeshPaintComponentAdapter.h
@@ -6,12 +6,13 @@
 #include "TMeshPaintOctree.h"
 #include "MeshAdapter.h"
 #include "Spatial/MeshAABBTree3.h"
+#include "UObject/GCObject.h"
 
 struct FIndex3i;
 typedef TIndexMeshArrayAdapter<uint32, double> FIndexMeshArrayAdapterd;
 
 /** Base mesh paint geometry adapter, handles basic sphere intersection using a Octree */
-class MESHPAINTINGTOOLSET_API FBaseMeshPaintComponentAdapter : public IMeshPaintComponentAdapter
+class MESHPAINTINGTOOLSET_API FBaseMeshPaintComponentAdapter : public IMeshPaintComponentAdapter, public FGCObject
 {
 public:
 	/** Start IMeshPaintGeometryAdapter Overrides */
diff --git a/Engine/Plugins/MeshPainting/Source/MeshPaintingToolset/Public/IMeshPaintComponentAdapter.h b/Engine/Plugins/MeshPainting/Source/MeshPaintingToolset/Public/IMeshPaintComponentAdapter.h
--- a/Engine/Plugins/MeshPainting/Source/MeshPaintingToolset/Public/IMeshPaintComponentAdapter.h
+++ b/Engine/Plugins/MeshPainting/Source/MeshPaintingToolset/Public/IMeshPaintComponentAdapter.h
@@ -47,9 +47,6 @@ public:
 
 	/** Applies or removes an override texture to use while rendering the meshes materials */
 	virtual void ApplyOrRemoveTextureOverride(UTexture* SourceTexture, UTexture* OverrideTexture) const = 0;
-
-	/** Reference collecting to prevent GC-ing */
-	virtual void AddReferencedObjects(FReferenceCollector& Collector) = 0;
 			
 	/** Returns the vertices for the current LOD index in the Mesh */
 	virtual const TArray<FVector>& GetMeshVertices() const = 0;
