Check that the mesh adapter is valid before using it.

diff --git a/Engine/Plugins/MeshPainting/Source/MeshPaintingToolset/Private/MeshSelect.cpp b/Engine/Plugins/MeshPainting/Source/MeshPaintingToolset/Private/MeshSelect.cpp
--- a/Engine/Plugins/MeshPainting/Source/MeshPaintingToolset/Private/MeshSelect.cpp
+++ b/Engine/Plugins/MeshPainting/Source/MeshPaintingToolset/Private/MeshSelect.cpp
@@ -98,7 +98,7 @@ void UMeshClickTool::OnClicked(const FInputDeviceRay& ClickPos)
 				{
 					UMeshComponent* MeshComponent = Cast<UMeshComponent>(CandidateComponent);
 					TSharedPtr<IMeshPaintComponentAdapter> MeshAdapter = FMeshPaintComponentAdapterFactory::CreateAdapterForMesh(MeshComponent, 0);
-					if (IsMeshAdapterSupported(MeshAdapter))
+					if (MeshAdapter.IsValid() && IsMeshAdapterSupported(MeshAdapter))
 					{
 						MeshToolManager->AddPaintableMeshComponent(MeshComponent);
 						MeshToolManager->AddToComponentToAdapterMap(MeshComponent, MeshAdapter);
