Making the mesh paint brushes adapt to the size of the selected mesh.

diff --git a/Engine/Plugins/MeshPainting/Source/MeshPaintingToolset/Private/MeshTexturePaintingTool.cpp b/Engine/Plugins/MeshPainting/Source/MeshPaintingToolset/Private/MeshTexturePaintingTool.cpp
--- a/Engine/Plugins/MeshPainting/Source/MeshPaintingToolset/Private/MeshTexturePaintingTool.cpp
+++ b/Engine/Plugins/MeshPainting/Source/MeshPaintingToolset/Private/MeshTexturePaintingTool.cpp
@@ -75,7 +75,6 @@ void UMeshTexturePaintingToolProperties::SaveProperties(UInteractiveTool* SaveFr
 
 void UMeshTexturePaintingToolProperties::RestoreProperties(UInteractiveTool* RestoreToTool)
 {
-	UBrushBaseProperties::RestoreProperties(RestoreToTool);
 	UMeshTexturePaintingToolProperties* PropertyCache = GetPropertyCache<UMeshTexturePaintingToolProperties>();
 	this->PaintColor = PropertyCache->PaintColor;
 	this->EraseColor = PropertyCache->EraseColor;
@@ -192,10 +191,39 @@ void UMeshTexturePaintingTool::Tick(float DeltaTime)
 
 void UMeshTexturePaintingTool::OnPropertyModified(UObject* PropertySet, FProperty* Property)
 {
+	Super::OnPropertyModified(PropertySet, Property);
 	bResultValid = false;
 }
 
 
+double UMeshTexturePaintingTool::EstimateMaximumTargetDimension()
+{
+	UMeshToolManager* MeshToolManager = Cast<UMeshToolManager>(GetToolManager());
+	bool bFoundComponentToUse = false;
+	FBoxSphereBounds Bounds = FBoxSphereBounds(0.0);
+	if (MeshToolManager)
+	{
+		// Preferentially use paintable components
+		for (UMeshComponent* PaintableComponent : MeshToolManager->GetPaintableMeshComponents())
+		{
+			bFoundComponentToUse = true;
+			Bounds = Bounds + PaintableComponent->Bounds;
+		}
+		// Otherwise use selected components
+		if (!bFoundComponentToUse)
+		{
+			for (UMeshComponent* SelectedComponent : MeshToolManager->GetSelectedMeshComponents())
+			{
+				bFoundComponentToUse = true;
+				Bounds = Bounds + SelectedComponent->Bounds;
+			}
+		}
+
+	}
+
+	return bFoundComponentToUse ? Bounds.GetBox().GetExtent().GetAbsMax() : Super::EstimateMaximumTargetDimension();
+}
+
 double UMeshTexturePaintingTool::CalculateTargetEdgeLength(int TargetTriCount)
 {
 	double TargetTriArea = InitialMeshArea / (double)TargetTriCount;
diff --git a/Engine/Plugins/MeshPainting/Source/MeshPaintingToolset/Private/MeshVertexPaintingTool.cpp b/Engine/Plugins/MeshPainting/Source/MeshPaintingToolset/Private/MeshVertexPaintingTool.cpp
--- a/Engine/Plugins/MeshPainting/Source/MeshPaintingToolset/Private/MeshVertexPaintingTool.cpp
+++ b/Engine/Plugins/MeshPainting/Source/MeshPaintingToolset/Private/MeshVertexPaintingTool.cpp
@@ -73,7 +73,6 @@ void UMeshVertexPaintingToolProperties::SaveProperties(UInteractiveTool* SaveFro
 
 void UMeshVertexPaintingToolProperties::RestoreProperties(UInteractiveTool* RestoreToTool)
 {
-	UBrushBaseProperties::RestoreProperties(RestoreToTool);
 	UMeshVertexPaintingToolProperties* PropertyCache = GetPropertyCache<UMeshVertexPaintingToolProperties>();
 	this->PaintColor = PropertyCache->PaintColor;
 	this->EraseColor = PropertyCache->EraseColor;
@@ -193,10 +192,40 @@ void UMeshVertexPaintingTool::Tick(float DeltaTime)
 
 void UMeshVertexPaintingTool::OnPropertyModified(UObject* PropertySet, FProperty* Property)
 {
+	Super::OnPropertyModified(PropertySet, Property);
 	bResultValid = false;
 }
 
 
+double UMeshVertexPaintingTool::EstimateMaximumTargetDimension()
+{
+	UMeshToolManager* MeshToolManager = Cast<UMeshToolManager>(GetToolManager());
+	bool bFoundComponentToUse = false;
+	FBoxSphereBounds Bounds = FBoxSphereBounds(0.0);
+	if (MeshToolManager)
+	{
+		// Preferentially use paintable components
+		for(UMeshComponent* PaintableComponent : MeshToolManager->GetPaintableMeshComponents())
+		{
+			bFoundComponentToUse = true;
+			Bounds = Bounds+PaintableComponent->Bounds;
+		}
+		// Otherwise use selected components
+		if (!bFoundComponentToUse)
+		{
+			for (UMeshComponent* SelectedComponent : MeshToolManager->GetSelectedMeshComponents())
+			{
+				bFoundComponentToUse = true;
+				Bounds = Bounds + SelectedComponent->Bounds;
+			}
+		}
+		
+	}
+
+	return bFoundComponentToUse ? Bounds.GetBox().GetExtent().GetAbsMax() : Super::EstimateMaximumTargetDimension();
+	
+}
+
 double UMeshVertexPaintingTool::CalculateTargetEdgeLength(int TargetTriCount)
 {
 	double TargetTriArea = InitialMeshArea / (double)TargetTriCount;
diff --git a/Engine/Plugins/MeshPainting/Source/MeshPaintingToolset/Public/MeshTexturePaintingTool.h b/Engine/Plugins/MeshPainting/Source/MeshPaintingToolset/Public/MeshTexturePaintingTool.h
--- a/Engine/Plugins/MeshPainting/Source/MeshPaintingToolset/Public/MeshTexturePaintingTool.h
+++ b/Engine/Plugins/MeshPainting/Source/MeshPaintingToolset/Public/MeshTexturePaintingTool.h
@@ -118,6 +118,7 @@ public:
 	virtual void OnEndDrag(const FRay& Ray) override;
 	virtual	bool HitTest(const FRay& Ray, FHitResult& OutHit) override;
 	virtual void OnPropertyModified(UObject* PropertySet, FProperty* Property) override;
+	virtual double EstimateMaximumTargetDimension() override;
 	virtual bool IsPainting() const
 	{
 		return bArePainting;
diff --git a/Engine/Plugins/MeshPainting/Source/MeshPaintingToolset/Public/MeshVertexPaintingTool.h b/Engine/Plugins/MeshPainting/Source/MeshPaintingToolset/Public/MeshVertexPaintingTool.h
--- a/Engine/Plugins/MeshPainting/Source/MeshPaintingToolset/Public/MeshVertexPaintingTool.h
+++ b/Engine/Plugins/MeshPainting/Source/MeshPaintingToolset/Public/MeshVertexPaintingTool.h
@@ -195,13 +195,14 @@ public:
 	{
 		return bArePainting;
 	}
-
+	virtual double EstimateMaximumTargetDimension() override;
 
 	FSimpleDelegate OnPaintingFinished()
 	{
 		return OnPaintingFinishedDelegate;
 	}
 
+
 protected:
 	virtual void SetAdditionalPaintParameters(FMeshPaintParameters& InPaintParameters) {};
 	virtual void FinishPainting();
