Fixing the mesh painting tools to allow navigation if they are not clicking or painting on a paintable component.

diff --git a/Engine/Plugins/MeshPainting/Source/MeshPaintEditorMode/Private/MeshPaintMode.cpp b/Engine/Plugins/MeshPainting/Source/MeshPaintEditorMode/Private/MeshPaintMode.cpp
--- a/Engine/Plugins/MeshPainting/Source/MeshPaintEditorMode/Private/MeshPaintMode.cpp
+++ b/Engine/Plugins/MeshPainting/Source/MeshPaintEditorMode/Private/MeshPaintMode.cpp
@@ -290,11 +290,6 @@ void UMeshPaintMode::OnToolEnded(UInteractiveToolManager* Manager, UInteractiveT
 	Super::OnToolEnded(Manager, Tool);
 }
 
-bool UMeshPaintMode::InputAxis(FEditorViewportClient* InViewportClient, FViewport* Viewport, int32 ControllerId, FKey Key, float Delta, float DeltaTime)
-{
-	return !InViewportClient->IsAltPressed();
-}
-
 void UMeshPaintMode::ActorSelectionChangeNotify()
 {
 	UpdateSelectedMeshes();
diff --git a/Engine/Plugins/MeshPainting/Source/MeshPaintEditorMode/Private/MeshPaintMode.h b/Engine/Plugins/MeshPainting/Source/MeshPaintEditorMode/Private/MeshPaintMode.h
--- a/Engine/Plugins/MeshPainting/Source/MeshPaintEditorMode/Private/MeshPaintMode.h
+++ b/Engine/Plugins/MeshPainting/Source/MeshPaintEditorMode/Private/MeshPaintMode.h
@@ -65,7 +65,6 @@ protected:
 	// UEdMode interface
 	virtual void OnToolStarted(UInteractiveToolManager* Manager, UInteractiveTool* Tool) override;
 	virtual void OnToolEnded(UInteractiveToolManager* Manager, UInteractiveTool* Tool) override;
-	virtual bool InputAxis(FEditorViewportClient* InViewportClient, FViewport* Viewport, int32 ControllerId, FKey Key, float Delta, float DeltaTime) override;
 	virtual void ActorSelectionChangeNotify() override;
 	virtual void ActivateDefaultTool() override;
 	virtual void UpdateOnPaletteChange() override;
diff --git a/Engine/Plugins/MeshPainting/Source/MeshPaintingToolset/Private/MeshSelect.cpp b/Engine/Plugins/MeshPainting/Source/MeshPaintingToolset/Private/MeshSelect.cpp
--- a/Engine/Plugins/MeshPainting/Source/MeshPaintingToolset/Private/MeshSelect.cpp
+++ b/Engine/Plugins/MeshPainting/Source/MeshPaintingToolset/Private/MeshSelect.cpp
@@ -69,47 +69,23 @@ void UMeshClickTool::OnUpdateModifierState(int ModifierID, bool bIsOn)
 
 FInputRayHit UMeshClickTool::IsHitByClick(const FInputDeviceRay& ClickPos)
 {
-	return FInputRayHit(0.0f);
+	if (UMeshToolManager* MeshToolManager = Cast<UMeshToolManager>(GetToolManager()))
+	{
+		FHitResult TraceHitResult = FindInitialHitResult(ClickPos, MeshToolManager);
+		if (TraceHitResult.bBlockingHit)
+		{
+			return FInputRayHit(TraceHitResult.Distance);
+		}
+	}
+	return FInputRayHit();
 }
 
 void UMeshClickTool::OnClicked(const FInputDeviceRay& ClickPos)
 {
 	if (UMeshToolManager* MeshToolManager = Cast<UMeshToolManager>(GetToolManager()))
 	{
-		const FVector& RayOrigin = ClickPos.WorldRay.Origin;
-		const FVector& RayDirection = ClickPos.WorldRay.Direction;
-		FRay Ray = FRay(RayOrigin, RayDirection);
-		const FVector TraceStart(RayOrigin);
-		const FVector TraceEnd(RayOrigin + RayDirection * HALF_WORLD_MAX);
-		MeshToolManager->ResetState();
-		//Default iterator only iterates over active levels.
-		const EActorIteratorFlags Flags = EActorIteratorFlags::SkipPendingKill;
-		// TODO: The tool needs to know the world its acting on
-		for (TActorIterator<AActor> It(GEditor->GetEditorWorldContext(false).World(), AActor::StaticClass(), Flags); It; ++It)
-		{
-			AActor* Actor = *It;
-			if (Actor->IsEditable() &&
-				Actor->IsListedInSceneOutliner() &&					// Only add actors that are allowed to be selected and drawn in editor
-				!Actor->IsTemplate() &&								// Should never happen, but we never want CDOs
-				!Actor->HasAnyFlags(RF_Transient))				// Don't add transient actors in non-play worlds
-			{
-				TArray<UActorComponent*> CandidateComponents = Actor->K2_GetComponentsByClass(UMeshComponent::StaticClass());
-				for (UActorComponent* CandidateComponent : CandidateComponents)
-				{
-					UMeshComponent* MeshComponent = Cast<UMeshComponent>(CandidateComponent);
-					TSharedPtr<IMeshPaintComponentAdapter> MeshAdapter = FMeshPaintComponentAdapterFactory::CreateAdapterForMesh(MeshComponent, 0);
-					if (MeshAdapter.IsValid() && IsMeshAdapterSupported(MeshAdapter))
-					{
-						MeshToolManager->AddPaintableMeshComponent(MeshComponent);
-						MeshToolManager->AddToComponentToAdapterMap(MeshComponent, MeshAdapter);
-					}
-				}
-					
-			}
-		}
-		FHitResult TraceHitResult(1.0f);
-		MeshToolManager->FindHitResult(Ray, TraceHitResult);
-		MeshToolManager->ResetState();
+		FHitResult TraceHitResult = FindInitialHitResult(ClickPos, MeshToolManager);
+
 		TArray<AActor*> SelectedActors;
 		if (UMeshComponent* MeshComponent = Cast<UMeshComponent>(TraceHitResult.GetComponent()))
 		{
@@ -137,6 +113,45 @@ void UMeshClickTool::OnClicked(const FInputDeviceRay& ClickPos)
 
 
 
+FHitResult UMeshClickTool::FindInitialHitResult(const FInputDeviceRay& ClickPos, class UMeshToolManager* MeshToolManager)
+{
+	const FVector& RayOrigin = ClickPos.WorldRay.Origin;
+	const FVector& RayDirection = ClickPos.WorldRay.Direction;
+	FRay Ray = FRay(RayOrigin, RayDirection);
+	const FVector TraceStart(RayOrigin);
+	const FVector TraceEnd(RayOrigin + RayDirection * HALF_WORLD_MAX);
+	MeshToolManager->ResetState();
+	//Default iterator only iterates over active levels.
+	const EActorIteratorFlags Flags = EActorIteratorFlags::SkipPendingKill;
+	// TODO: The tool needs to know the world its acting on
+	for (TActorIterator<AActor> It(GEditor->GetEditorWorldContext(false).World(), AActor::StaticClass(), Flags); It; ++It)
+	{
+		AActor* Actor = *It;
+		if (Actor->IsEditable() &&
+			Actor->IsListedInSceneOutliner() &&					// Only add actors that are allowed to be selected and drawn in editor
+			!Actor->IsTemplate() &&								// Should never happen, but we never want CDOs
+			!Actor->HasAnyFlags(RF_Transient))				// Don't add transient actors in non-play worlds
+		{
+			TArray<UActorComponent*> CandidateComponents = Actor->K2_GetComponentsByClass(UMeshComponent::StaticClass());
+			for (UActorComponent* CandidateComponent : CandidateComponents)
+			{
+				UMeshComponent* MeshComponent = Cast<UMeshComponent>(CandidateComponent);
+				TSharedPtr<IMeshPaintComponentAdapter> MeshAdapter = FMeshPaintComponentAdapterFactory::CreateAdapterForMesh(MeshComponent, 0);
+				if (MeshAdapter.IsValid() && IsMeshAdapterSupported(MeshAdapter))
+				{
+					MeshToolManager->AddPaintableMeshComponent(MeshComponent);
+					MeshToolManager->AddToComponentToAdapterMap(MeshComponent, MeshAdapter);
+				}
+			}
+
+		}
+	}
+	FHitResult TraceHitResult(1.0f);
+	MeshToolManager->FindHitResult(Ray, TraceHitResult);
+	MeshToolManager->ResetState();
+	return TraceHitResult;
+}
+
 UVertexAdapterClickTool::UVertexAdapterClickTool()
 	: UMeshClickTool()
 {
diff --git a/Engine/Plugins/MeshPainting/Source/MeshPaintingToolset/Private/MeshTexturePaintingTool.cpp b/Engine/Plugins/MeshPainting/Source/MeshPaintingToolset/Private/MeshTexturePaintingTool.cpp
--- a/Engine/Plugins/MeshPainting/Source/MeshPaintingToolset/Private/MeshTexturePaintingTool.cpp
+++ b/Engine/Plugins/MeshPainting/Source/MeshPaintingToolset/Private/MeshTexturePaintingTool.cpp
@@ -203,25 +203,27 @@ double UMeshTexturePaintingTool::EstimateMaximumTargetDimension()
 	FBoxSphereBounds Bounds = FBoxSphereBounds(0.0);
 	if (MeshToolManager)
 	{
-		// Preferentially use paintable components
-		for (UMeshComponent* PaintableComponent : MeshToolManager->GetPaintableMeshComponents())
-		{
-			bFoundComponentToUse = true;
-			Bounds = Bounds + PaintableComponent->Bounds;
-		}
-		// Otherwise use selected components
-		if (!bFoundComponentToUse)
+		bool bFirstItem = true;
+
+		FBoxSphereBounds Extents;
+		for (UMeshComponent* SelectedComponent : MeshToolManager->GetSelectedMeshComponents())
 		{
-			for (UMeshComponent* SelectedComponent : MeshToolManager->GetSelectedMeshComponents())
+			if (bFirstItem)
 			{
-				bFoundComponentToUse = true;
-				Bounds = Bounds + SelectedComponent->Bounds;
+				Extents = SelectedComponent->Bounds;
 			}
+			else
+			{
+				Extents = Extents + SelectedComponent->Bounds;
+			}
+
+			bFirstItem = false;
 		}
 
+		return Extents.BoxExtent.GetAbsMax();
 	}
 
-	return bFoundComponentToUse ? Bounds.GetBox().GetExtent().GetAbsMax() : Super::EstimateMaximumTargetDimension();
+	return Super::EstimateMaximumTargetDimension();
 }
 
 double UMeshTexturePaintingTool::CalculateTargetEdgeLength(int TargetTriCount)
@@ -526,8 +528,9 @@ bool UMeshTexturePaintingTool::HitTest(const FRay& Ray, FHitResult& OutHit)
 	bool bUsed = false;
 	if (UMeshToolManager* MeshToolManager = Cast<UMeshToolManager>(GetToolManager()))
 	{
-		bUsed = MeshToolManager->FindHitResult(Ray, OutHit);
+		MeshToolManager->FindHitResult(Ray, OutHit);
 		LastBestHitResult = OutHit;
+		bUsed = OutHit.bBlockingHit;
 	}
 	return bUsed;
 }
diff --git a/Engine/Plugins/MeshPainting/Source/MeshPaintingToolset/Private/MeshVertexPaintingTool.cpp b/Engine/Plugins/MeshPainting/Source/MeshPaintingToolset/Private/MeshVertexPaintingTool.cpp
--- a/Engine/Plugins/MeshPainting/Source/MeshPaintingToolset/Private/MeshVertexPaintingTool.cpp
+++ b/Engine/Plugins/MeshPainting/Source/MeshPaintingToolset/Private/MeshVertexPaintingTool.cpp
@@ -204,25 +204,27 @@ double UMeshVertexPaintingTool::EstimateMaximumTargetDimension()
 	FBoxSphereBounds Bounds = FBoxSphereBounds(0.0);
 	if (MeshToolManager)
 	{
-		// Preferentially use paintable components
-		for(UMeshComponent* PaintableComponent : MeshToolManager->GetPaintableMeshComponents())
-		{
-			bFoundComponentToUse = true;
-			Bounds = Bounds+PaintableComponent->Bounds;
-		}
-		// Otherwise use selected components
-		if (!bFoundComponentToUse)
+		bool bFirstItem = true;
+
+		FBoxSphereBounds Extents;
+		for (UMeshComponent* SelectedComponent : MeshToolManager->GetSelectedMeshComponents())
 		{
-			for (UMeshComponent* SelectedComponent : MeshToolManager->GetSelectedMeshComponents())
+			if (bFirstItem)
+			{
+				Extents = SelectedComponent->Bounds;
+			}
+			else
 			{
-				bFoundComponentToUse = true;
-				Bounds = Bounds + SelectedComponent->Bounds;
+				Extents = Extents + SelectedComponent->Bounds;
 			}
+
+			bFirstItem = false;
 		}
-		
+
+		return Extents.BoxExtent.GetAbsMax();
 	}
 
-	return bFoundComponentToUse ? Bounds.GetBox().GetExtent().GetAbsMax() : Super::EstimateMaximumTargetDimension();
+	return Super::EstimateMaximumTargetDimension();
 	
 }
 
@@ -488,8 +490,9 @@ bool UMeshVertexPaintingTool::HitTest(const FRay& Ray, FHitResult& OutHit)
 	bool bUsed = false;
 	if (UMeshToolManager* MeshToolManager = Cast<UMeshToolManager>(GetToolManager()))
 	{
-		bUsed = MeshToolManager->FindHitResult(Ray, OutHit);
+		MeshToolManager->FindHitResult(Ray, OutHit);
 		LastBestHitResult = OutHit;
+		bUsed = OutHit.bBlockingHit;
 	}
 	return bUsed;
 }
diff --git a/Engine/Plugins/MeshPainting/Source/MeshPaintingToolset/Public/MeshSelect.h b/Engine/Plugins/MeshPainting/Source/MeshPaintingToolset/Public/MeshSelect.h
--- a/Engine/Plugins/MeshPainting/Source/MeshPaintingToolset/Public/MeshSelect.h
+++ b/Engine/Plugins/MeshPainting/Source/MeshPaintingToolset/Public/MeshSelect.h
@@ -62,7 +62,6 @@ public:
 	virtual void OnUpdateModifierState(int ModifierID, bool bIsOn) override;
 	virtual FInputRayHit IsHitByClick(const FInputDeviceRay& ClickPos) override;
 	virtual void OnClicked(const FInputDeviceRay& ClickPos) override;
-	
 	virtual bool IsMeshAdapterSupported(TSharedPtr<IMeshPaintComponentAdapter> MeshAdapter)
 	{
 		return true;
@@ -73,6 +72,9 @@ protected:
 	static const int AdditiveSelectionModifier = 1;
 	bool bAddToSelectionSet;
 
+private:
+	FHitResult FindInitialHitResult(const FInputDeviceRay& ClickPos, class UMeshToolManager* MeshToolManager);
+
 };
 
 UCLASS()
