Turn off the paint tool and current/next texture if there aren't any paintable textures on the mesh.

diff --git a/Engine/Plugins/MeshPainting/Source/MeshPaintEditorMode/Private/MeshPaintMode.cpp b/Engine/Plugins/MeshPainting/Source/MeshPaintEditorMode/Private/MeshPaintMode.cpp
--- a/Engine/Plugins/MeshPainting/Source/MeshPaintEditorMode/Private/MeshPaintMode.cpp
+++ b/Engine/Plugins/MeshPainting/Source/MeshPaintEditorMode/Private/MeshPaintMode.cpp
@@ -264,8 +264,16 @@ void UMeshPaintMode::BindCommands()
 			FCanExecuteAction::CreateUObject(this, &UMeshPaintMode::CanPropagateVertexColorsToLODs)
 		));
 
-	CommandList->MapAction(Commands.NextTexture, FExecuteAction::CreateUObject(this, &UMeshPaintMode::CycleTextures, 1));
-	CommandList->MapAction(Commands.PreviousTexture, FExecuteAction::CreateUObject(this, &UMeshPaintMode::CycleTextures, -1));
+	CommandList->MapAction(Commands.NextTexture, 
+		FUIAction(
+			FExecuteAction::CreateUObject(this, &UMeshPaintMode::CycleTextures, 1),
+			FCanExecuteAction::CreateUObject(this, &UMeshPaintMode::CanCycleTextures)
+		));
+	CommandList->MapAction(Commands.PreviousTexture,
+		FUIAction(
+			FExecuteAction::CreateUObject(this, &UMeshPaintMode::CycleTextures, -1),
+			FCanExecuteAction::CreateUObject(this, &UMeshPaintMode::CanCycleTextures)
+		));
 	CommandList->MapAction(Commands.SaveTexturePaint, FUIAction(FExecuteAction::CreateStatic(&UMeshPaintModeHelpers::SaveModifiedTextures), FCanExecuteAction::CreateStatic(&UMeshPaintModeHelpers::CanSaveModifiedTextures)));
 
 	auto HasPaintChanges = [this]() -> bool { return GetNumberOfPendingPaintChanges() > 0; };
@@ -698,22 +706,28 @@ void UMeshPaintMode::UpdateCachedVertexDataSize()
 
 void UMeshPaintMode::CycleMeshLODs(int32 Direction)
 {
-	FScopedTransaction Transaction(LOCTEXT("LevelMeshPainter_CycleLOD", "Changed Current LOD"));
 	if (UMeshColorPaintingTool* ColorPaintingTool = Cast<UMeshColorPaintingTool>(GetToolManager()->GetActiveTool(EToolSide::Left)))
 	{
+		FScopedTransaction Transaction(LOCTEXT("LevelMeshPainter_CycleLOD", "Changed Current LOD"));
 		ColorPaintingTool->CycleMeshLODs(Direction);
 	}
 }
 
 void UMeshPaintMode::CycleTextures(int32 Direction)
 {
-	FScopedTransaction Transaction(LOCTEXT("LevelMeshPainter_CycleTexture", "Changed Current Texture"));
 	if (UMeshTexturePaintingTool* TexturePaintingTool = Cast<UMeshTexturePaintingTool>(GetToolManager()->GetActiveTool(EToolSide::Left)))
 	{
+		FScopedTransaction Transaction(LOCTEXT("LevelMeshPainter_CycleTexture", "Changed Current Texture"));
 		TexturePaintingTool->CycleTextures(Direction);
 	}
 }
 
+bool UMeshPaintMode::CanCycleTextures() const
+{
+	UMeshTexturePaintingTool* TexturePaintingTool = Cast<UMeshTexturePaintingTool>(GetToolManager()->GetActiveTool(EToolSide::Left));
+	return TexturePaintingTool != nullptr;
+}
+
 void UMeshPaintMode::CommitAllPaintedTextures()
 {
 	if (UMeshTexturePaintingTool* TexturePaintingTool = Cast<UMeshTexturePaintingTool>(GetToolManager()->GetActiveTool(EToolSide::Left)))
diff --git a/Engine/Plugins/MeshPainting/Source/MeshPaintEditorMode/Private/MeshPaintMode.h b/Engine/Plugins/MeshPainting/Source/MeshPaintEditorMode/Private/MeshPaintMode.h
--- a/Engine/Plugins/MeshPainting/Source/MeshPaintEditorMode/Private/MeshPaintMode.h
+++ b/Engine/Plugins/MeshPainting/Source/MeshPaintEditorMode/Private/MeshPaintMode.h
@@ -91,6 +91,7 @@ protected:
 	void UpdateCachedVertexDataSize();
 	void CycleMeshLODs(int32 Direction);
 	void CycleTextures(int32 Direction);
+	bool CanCycleTextures() const;
 	void CommitAllPaintedTextures();
 	int32 GetNumberOfPendingPaintChanges();
 
diff --git a/Engine/Plugins/MeshPainting/Source/MeshPaintingToolset/Private/MeshTexturePaintingTool.cpp b/Engine/Plugins/MeshPainting/Source/MeshPaintingToolset/Private/MeshTexturePaintingTool.cpp
--- a/Engine/Plugins/MeshPainting/Source/MeshPaintingToolset/Private/MeshTexturePaintingTool.cpp
+++ b/Engine/Plugins/MeshPainting/Source/MeshPaintingToolset/Private/MeshTexturePaintingTool.cpp
@@ -17,6 +17,7 @@
 #include "CanvasItem.h"
 #include "ScopedTransaction.h"
 #include "BaseGizmos/BrushStampIndicator.h"
+#include "MeshPaintAdapterFactory.h"
 
 #define LOCTEXT_NAMESPACE "MeshTextureBrush"
 
@@ -29,7 +30,19 @@
 bool UMeshTexturePaintingToolBuilder::CanBuildTool(const FToolBuilderState& SceneState) const
 {
 	UActorComponent* ActorComponent = ToolBuilderUtil::FindFirstComponent(SceneState, UMeshPaintingToolset::HasPaintableMesh);
-	return ActorComponent != nullptr;
+	UMeshToolManager* MeshToolManager = Cast<UMeshToolManager>(SceneState.ToolManager);
+	TArray<FPaintableTexture> PaintableTextures;
+	if (ActorComponent && MeshToolManager)
+	{
+		UMeshComponent* MeshComponent = Cast<UMeshComponent>(ActorComponent);
+		if (MeshComponent)
+		{
+			TSharedPtr<IMeshPaintComponentAdapter> MeshAdapter = FMeshPaintComponentAdapterFactory::CreateAdapterForMesh(MeshComponent, 0);
+			UTexturePaintToolset::RetrieveTexturesForComponent(MeshComponent, MeshAdapter.Get(), PaintableTextures);
+		}
+	}
+
+	return PaintableTextures.Num() > 0;
 }
 
 UInteractiveTool* UMeshTexturePaintingToolBuilder::BuildTool(const FToolBuilderState& SceneState) const
