Resolve some issues with Undo/Redo in Paint Mode.
- Removed automatic starting of "default" tool in UEdMode::Tick
- FModeToolkit::OnToolPaletteChanged no longer automatically accepts current tool
- fixed bug in FEditorModeTools::InvokeToolPaletteTab where UEdMode wasn't handled properly
- UMeshPaintMode::Tick now switches to tab that contains active tool, if it is not visible  (happens during undo/redo)
- UMeshPaintMode::UpdateOnPaletteChange now handles transition to appropriate paint/select Tool for new tab
#rnx

diff --git a/Engine/Plugins/MeshPainting/Source/MeshPaintEditorMode/Private/MeshPaintMode.cpp b/Engine/Plugins/MeshPainting/Source/MeshPaintEditorMode/Private/MeshPaintMode.cpp
--- a/Engine/Plugins/MeshPainting/Source/MeshPaintEditorMode/Private/MeshPaintMode.cpp
+++ b/Engine/Plugins/MeshPainting/Source/MeshPaintEditorMode/Private/MeshPaintMode.cpp
@@ -32,6 +32,13 @@ FName UMeshPaintMode::MeshPaintMode_Color = FName(TEXT("Color"));
 FName UMeshPaintMode::MeshPaintMode_Weights = FName(TEXT("Weights"));
 FName UMeshPaintMode::MeshPaintMode_Texture = FName(TEXT("Texture"));
 
+FString UMeshPaintMode::VertexSelectToolName = TEXT("VertexAdapterClickTool");
+FString UMeshPaintMode::TextureSelectToolName = TEXT("TextureAdapterClickTool");
+FString UMeshPaintMode::ColorPaintToolName = TEXT("ColorBrushTool");
+FString UMeshPaintMode::WeightPaintToolName = TEXT("WeightBrushTool");
+FString UMeshPaintMode::TexturePaintToolName = TEXT("TextureBrushTool");
+
+
 UMeshVertexPaintingToolProperties* UMeshPaintMode::GetVertexToolProperties()
 {
 	if (GLevelEditorModeTools().GetActiveScriptableMode("MeshPaintMode")
@@ -120,11 +127,11 @@ void UMeshPaintMode::Enter()
 	ModeSettings = Cast<UMeshPaintModeSettings>(SettingsObject);
 	
 	FMeshPaintEditorModeCommands ToolManagerCommands = FMeshPaintEditorModeCommands::Get();
-	RegisterTool(ToolManagerCommands.VertexSelect, TEXT("VertexAdapterClickTool"), NewObject<UVertexAdapterClickToolBuilder>());
-	RegisterTool(ToolManagerCommands.TextureSelect, TEXT("TextureAdapterClickTool"), NewObject<UTextureAdapterClickToolBuilder>());
-	RegisterTool(ToolManagerCommands.ColorPaint, TEXT("ColorBrushTool"), NewObject<UMeshColorPaintingToolBuilder>());
-	RegisterTool(ToolManagerCommands.WeightPaint, TEXT("WeightBrushTool"), NewObject<UMeshWeightPaintingToolBuilder>());
-	RegisterTool(ToolManagerCommands.TexturePaint, TEXT("TextureBrushTool"), NewObject<UMeshTexturePaintingToolBuilder>());
+	RegisterTool(ToolManagerCommands.VertexSelect, VertexSelectToolName, NewObject<UVertexAdapterClickToolBuilder>());
+	RegisterTool(ToolManagerCommands.TextureSelect, TextureSelectToolName, NewObject<UTextureAdapterClickToolBuilder>());
+	RegisterTool(ToolManagerCommands.ColorPaint, ColorPaintToolName, NewObject<UMeshColorPaintingToolBuilder>());
+	RegisterTool(ToolManagerCommands.WeightPaint, WeightPaintToolName, NewObject<UMeshWeightPaintingToolBuilder>());
+	RegisterTool(ToolManagerCommands.TexturePaint, TexturePaintToolName, NewObject<UMeshTexturePaintingToolBuilder>());
 	UpdateSelectedMeshes();
 
 	// disable tool change tracking to activate default tool, and then switch to full undo/redo tracking mode
@@ -189,6 +196,28 @@ void UMeshPaintMode::Tick(FEditorViewportClient* ViewportClient, float DeltaTime
 	{
 		UpdateCachedVertexDataSize();
 	}
+
+	// Make sure that correct tab is visible for the current tool
+	// Note that currently Color and Weight mode share the same Select tool
+	FString ActiveTool = GetToolManager()->GetActiveToolName(EToolSide::Mouse);
+	FName ActiveTab = GetCurrentPaletteName();
+	FName TargetTab = ActiveTab;
+	if (ActiveTool == TexturePaintToolName || ActiveTool == TextureSelectToolName)
+	{
+		TargetTab = MeshPaintMode_Texture;
+	}
+	else if (ActiveTool == ColorPaintToolName)
+	{
+		TargetTab = MeshPaintMode_Color;
+	}
+	else if (ActiveTool == WeightPaintToolName)
+	{
+		TargetTab = MeshPaintMode_Weights;
+	}
+	if ( TargetTab != ActiveTab)
+	{
+		GetModeManager()->InvokeToolPaletteTab(GetID(), TargetTab);
+	}
 }
 
 TMap<FName, TArray<TSharedPtr<FUICommandInfo>>> UMeshPaintMode::GetModeCommands() const
@@ -748,25 +777,49 @@ int32 UMeshPaintMode::GetNumberOfPendingPaintChanges()
 
 void UMeshPaintMode::ActivateDefaultTool()
 {
-	if (CurrentPaletteName == UMeshPaintMode::MeshPaintMode_Color)
+	if (GetCurrentPaletteName() == UMeshPaintMode::MeshPaintMode_Color)
 	{
-		ToolsContext->StartTool(EToolSide::Mouse, TEXT("VertexAdapterClickTool"));
+		ToolsContext->StartTool(EToolSide::Mouse, VertexSelectToolName);
 	}
-	if (CurrentPaletteName == UMeshPaintMode::MeshPaintMode_Weights)
+	if (GetCurrentPaletteName() == UMeshPaintMode::MeshPaintMode_Weights)
 	{
-		ToolsContext->StartTool(EToolSide::Mouse, TEXT("VertexAdapterClickTool"));
+		ToolsContext->StartTool(EToolSide::Mouse, VertexSelectToolName);
 	}
-	if (CurrentPaletteName == UMeshPaintMode::MeshPaintMode_Texture)
+	if (GetCurrentPaletteName() == UMeshPaintMode::MeshPaintMode_Texture)
 	{
-		ToolsContext->StartTool(EToolSide::Mouse, TEXT("TextureAdapterClickTool"));
+		ToolsContext->StartTool(EToolSide::Mouse, TextureSelectToolName);
 	}
 }
 
 void UMeshPaintMode::UpdateOnPaletteChange()
 {
 	UpdateSelectedMeshes();
+
+	FString ActiveTool = GetToolManager()->GetActiveToolName(EToolSide::Mouse);
+	bool bInAnyPaintTool = (ActiveTool == WeightPaintToolName || ActiveTool == TexturePaintToolName || ActiveTool == ColorPaintToolName);
+
+	// figure out which tool we would like to be in based on currently-active tool
+	FString SwitchToTool;
+	if (GetCurrentPaletteName() == UMeshPaintMode::MeshPaintMode_Color)
+	{
+		SwitchToTool = (bInAnyPaintTool) ? ColorPaintToolName : VertexSelectToolName;
+	}
+	else if (GetCurrentPaletteName() == UMeshPaintMode::MeshPaintMode_Weights)
+	{
+		SwitchToTool = (bInAnyPaintTool) ? WeightPaintToolName: VertexSelectToolName;
+	}
+	else if (GetCurrentPaletteName() == UMeshPaintMode::MeshPaintMode_Texture)
+	{
+		SwitchToTool = (bInAnyPaintTool) ? TexturePaintToolName : TextureSelectToolName;
+	}
+	// change to new tool if it is different
+	if (SwitchToTool.IsEmpty() == false && SwitchToTool != ActiveTool)
+	{
+		ToolsContext->StartTool(EToolSide::Mouse, SwitchToTool);
+	}
 }
 
+
 void UMeshPaintMode::OnResetViewMode()
 {
 	// Reset viewport color mode for all active viewports
diff --git a/Engine/Plugins/MeshPainting/Source/MeshPaintEditorMode/Private/MeshPaintMode.h b/Engine/Plugins/MeshPainting/Source/MeshPaintEditorMode/Private/MeshPaintMode.h
--- a/Engine/Plugins/MeshPainting/Source/MeshPaintEditorMode/Private/MeshPaintMode.h
+++ b/Engine/Plugins/MeshPainting/Source/MeshPaintEditorMode/Private/MeshPaintMode.h
@@ -46,6 +46,11 @@ public:
 	static FName MeshPaintMode_Color;
 	static FName MeshPaintMode_Texture;
 	static FName MeshPaintMode_Weights;
+	static FString VertexSelectToolName;
+	static FString TextureSelectToolName;
+	static FString ColorPaintToolName;
+	static FString WeightPaintToolName;
+	static FString TexturePaintToolName;
 
 	virtual TMap<FName, TArray<TSharedPtr<FUICommandInfo>>> GetModeCommands() const override;
 	/** Returns the instance of ComponentClass found in the current Editor selection */
diff --git a/Engine/Source/Editor/UnrealEd/Private/EditorModeManager.cpp b/Engine/Source/Editor/UnrealEd/Private/EditorModeManager.cpp
--- a/Engine/Source/Editor/UnrealEd/Private/EditorModeManager.cpp
+++ b/Engine/Source/Editor/UnrealEd/Private/EditorModeManager.cpp
@@ -837,11 +837,18 @@ void FEditorModeTools::InvokeToolPaletteTab(FEditorModeID InModeID, FName InPale
 		{
 			TSharedRef<SWidget> PaletteWidget = Row.ToolbarWidget.ToSharedRef();
 
-			FEdMode* Mode = GetActiveMode(InModeID);
-			TSharedPtr<FModeToolkit> RowToolkit = Mode->GetToolkit();
+			TSharedPtr<FModeToolkit> RowToolkit;
+			if (FEdMode* Mode = GetActiveMode(InModeID))
+			{
+				RowToolkit = Mode->GetToolkit();
+			}
+			else if (UEdMode* ScriptableMode = GetActiveScriptableMode(InModeID))
+			{
+				RowToolkit = ScriptableMode->GetToolkit();
+			}
 
 			TSharedPtr<SWidget> ActiveWidget = ModeToolbarPaletteSwitcher.Pin()->GetActiveWidget();
-			if (ActiveWidget.Get() != Row.ToolbarWidget.Get())
+			if (RowToolkit && ActiveWidget.Get() != Row.ToolbarWidget.Get())
 			{
 				ModeToolbarPaletteSwitcher.Pin()->SetActiveWidget(Row.ToolbarWidget.ToSharedRef());
 				RowToolkit->OnToolPaletteChanged(Row.PaletteName);
diff --git a/Engine/Source/Editor/UnrealEd/Private/Toolkits/BaseToolkit.cpp b/Engine/Source/Editor/UnrealEd/Private/Toolkits/BaseToolkit.cpp
--- a/Engine/Source/Editor/UnrealEd/Private/Toolkits/BaseToolkit.cpp
+++ b/Engine/Source/Editor/UnrealEd/Private/Toolkits/BaseToolkit.cpp
@@ -419,11 +419,6 @@ void FModeToolkit::OnToolPaletteChanged(FName PaletteName)
 	if (UEdMode* EditorMode = GetScriptableEditorMode())
 	{
 		EditorMode->SetCurrentPaletteName(PaletteName);
-		if (UInteractiveTool* CurrentTool = EditorMode->GetToolManager()->GetActiveTool(EToolSide::Left))
-		{
-			// When swapping tools, default to accepting the active tool
-			EditorMode->GetToolManager()->DeactivateTool(EToolSide::Left, EToolShutdownType::Accept);
-		}
 	}
 }
 
diff --git a/Engine/Source/Editor/UnrealEd/Private/Tools/UEdMode.cpp b/Engine/Source/Editor/UnrealEd/Private/Tools/UEdMode.cpp
--- a/Engine/Source/Editor/UnrealEd/Private/Tools/UEdMode.cpp
+++ b/Engine/Source/Editor/UnrealEd/Private/Tools/UEdMode.cpp
@@ -396,7 +396,6 @@ UEdMode::UEdMode()
 	ToolsContext = nullptr;
 	ToolsContextClass = UInteractiveToolsContext::StaticClass();
 	ToolCommandList = MakeShareable(new FUICommandList);
-	bCheckIfDefaultToolNeeded = false;
 }
 
 
@@ -715,11 +714,6 @@ void UEdMode::Tick(FEditorViewportClient* ViewportClient, float DeltaTime)
 		ToolsContext->ToolManager->Tick(DeltaTime);
 		ToolsContext->GizmoManager->Tick(DeltaTime);
 
-		if(bCheckIfDefaultToolNeeded && ToolsContext->ToolManager->GetActiveTool(EToolSide::Left) == nullptr)
-		{
-			ActivateDefaultTool();
-		}
-
 		if (bInvalidationPending)
 		{
 			ViewportClient->Invalidate();
diff --git a/Engine/Source/Editor/UnrealEd/Public/Tools/UEdMode.h b/Engine/Source/Editor/UnrealEd/Public/Tools/UEdMode.h
--- a/Engine/Source/Editor/UnrealEd/Public/Tools/UEdMode.h
+++ b/Engine/Source/Editor/UnrealEd/Public/Tools/UEdMode.h
@@ -314,10 +314,15 @@ public:
 		UpdateOnPaletteChange();
 	};
 
+	FName GetCurrentPaletteName() const
+	{
+		return CurrentPaletteName;
+	}
+
 protected:
 	virtual void CreateToolkit();
 	virtual void OnToolStarted(UInteractiveToolManager* Manager, UInteractiveTool* Tool) {};
-	virtual void OnToolEnded(UInteractiveToolManager* Manager, UInteractiveTool* Tool) { bCheckIfDefaultToolNeeded = true; };
+	virtual void OnToolEnded(UInteractiveToolManager* Manager, UInteractiveTool* Tool) {};
 	virtual void ActivateDefaultTool() {};
 	virtual void UpdateOnPaletteChange() {};
 
@@ -349,6 +354,5 @@ protected:
 	FOnModeToolNotification OnToolWarningMessage;
 
 	FName CurrentPaletteName;
-	bool bCheckIfDefaultToolNeeded;
 };
 
