Mesh click tool now uses hit proxies to find initial clicked actors

diff --git a/Engine/Plugins/MeshPainting/Source/MeshPaintingToolset/Private/MeshSelect.cpp b/Engine/Plugins/MeshPainting/Source/MeshPaintingToolset/Private/MeshSelect.cpp
--- a/Engine/Plugins/MeshPainting/Source/MeshPaintingToolset/Private/MeshSelect.cpp
+++ b/Engine/Plugins/MeshPainting/Source/MeshPaintingToolset/Private/MeshSelect.cpp
@@ -11,6 +11,9 @@
 #include "MeshPaintAdapterFactory.h"
 #include "EngineUtils.h"
 #include "Editor.h"
+#if WITH_EDITOR
+#include "HitProxies.h"
+#endif
 
 #define LOCTEXT_NAMESPACE "MeshSelection"
 
@@ -57,6 +60,15 @@ void UMeshClickTool::Setup()
 	GetToolManager()->DisplayMessage(
 		LOCTEXT("OnStartMeshSelectTool", "Select a mesh. Switch tools to paint vertex colors, blend between textures, or paint directly onto a texture file."),
 		EToolMessageLevel::UserNotification);
+
+	GetToolManager()->BeginUndoTransaction(LOCTEXT("MeshSelection", "Select Mesh"));
+
+
+	FSelectedOjectsChangeList NewSelection;
+	// TODO add CTRL handling
+	NewSelection.ModificationType = ESelectedObjectsModificationType::Clear;
+	GetToolManager()->RequestSelectionChange(NewSelection);
+	GetToolManager()->EndUndoTransaction();
 }
 
 void UMeshClickTool::OnUpdateModifierState(int ModifierID, bool bIsOn)
@@ -71,11 +83,12 @@ FInputRayHit UMeshClickTool::IsHitByClick(const FInputDeviceRay& ClickPos)
 {
 	if (UMeshToolManager* MeshToolManager = Cast<UMeshToolManager>(GetToolManager()))
 	{
-		FHitResult TraceHitResult = FindInitialHitResult(ClickPos, MeshToolManager);
-		if (TraceHitResult.bBlockingHit)
+		if (!bAddToSelectionSet || !AllowsMultiselect())
 		{
-			return FInputRayHit(TraceHitResult.Distance);
+			CachedClickedComponents.Empty();
+			CachedClickedActors.Empty();
 		}
+		return FindClickedComponentsAndCacheAdapters(ClickPos, MeshToolManager) ? FInputRayHit(0.0f) : FInputRayHit();
 	}
 	return FInputRayHit();
 }
@@ -84,54 +97,40 @@ void UMeshClickTool::OnClicked(const FInputDeviceRay& ClickPos)
 {
 	if (UMeshToolManager* MeshToolManager = Cast<UMeshToolManager>(GetToolManager()))
 	{
-		FHitResult TraceHitResult = FindInitialHitResult(ClickPos, MeshToolManager);
-
-		TArray<AActor*> SelectedActors;
-		if (UMeshComponent* MeshComponent = Cast<UMeshComponent>(TraceHitResult.GetComponent()))
+		for (UMeshComponent* MeshComponent : CachedClickedComponents)
 		{
-			TSharedPtr<IMeshPaintComponentAdapter> MeshAdapter = FMeshPaintComponentAdapterFactory::CreateAdapterForMesh(MeshComponent, 0);
+			TSharedPtr<IMeshPaintComponentAdapter> MeshAdapter = MeshToolManager->GetAdapterForComponent(MeshComponent);
 			if (MeshComponent->IsVisible() && MeshAdapter.IsValid() && MeshAdapter->IsValid() && IsMeshAdapterSupported(MeshAdapter))
 			{
 				MeshToolManager->AddPaintableMeshComponent(MeshComponent);
-				MeshToolManager->AddToComponentToAdapterMap(MeshComponent, MeshAdapter);
-				SelectedActors.Add(Cast<AActor>(MeshComponent->GetOuter()));
 				MeshAdapter->OnAdded();
 			}
-		}
 
-		GetToolManager()->BeginUndoTransaction(LOCTEXT("MeshSelection", "Select Mesh"));
+			GetToolManager()->BeginUndoTransaction(LOCTEXT("MeshSelection", "Select Mesh"));
 
 
-		FSelectedOjectsChangeList NewSelection;
-		// TODO add CTRL handling
-		NewSelection.ModificationType = bAddToSelectionSet ? ESelectedObjectsModificationType::Add : ESelectedObjectsModificationType::Replace;
-		NewSelection.Actors.Append(SelectedActors);
-		GetToolManager()->RequestSelectionChange(NewSelection);
-		GetToolManager()->EndUndoTransaction();
+			FSelectedOjectsChangeList NewSelection;
+			// TODO add CTRL handling
+			NewSelection.ModificationType = bAddToSelectionSet ? ESelectedObjectsModificationType::Add : ESelectedObjectsModificationType::Replace;
+			NewSelection.Actors.Append(CachedClickedActors);
+			GetToolManager()->RequestSelectionChange(NewSelection);
+			GetToolManager()->EndUndoTransaction();
+		}
 	}
 }
 
 
 
-FHitResult UMeshClickTool::FindInitialHitResult(const FInputDeviceRay& ClickPos, class UMeshToolManager* MeshToolManager)
+bool UMeshClickTool::FindClickedComponentsAndCacheAdapters(const FInputDeviceRay& ClickPos, class UMeshToolManager* MeshToolManager)
 {
-	const FVector& RayOrigin = ClickPos.WorldRay.Origin;
-	const FVector& RayDirection = ClickPos.WorldRay.Direction;
-	FRay Ray = FRay(RayOrigin, RayDirection);
-	const FVector TraceStart(RayOrigin);
-	const FVector TraceEnd(RayOrigin + RayDirection * HALF_WORLD_MAX);
-	MeshToolManager->ResetState();
-	//Default iterator only iterates over active levels.
-	const EActorIteratorFlags Flags = EActorIteratorFlags::SkipPendingKill;
-	// TODO: The tool needs to know the world its acting on
-	for (TActorIterator<AActor> It(GEditor->GetEditorWorldContext(false).World(), AActor::StaticClass(), Flags); It; ++It)
+	bool bFoundValidComponents = false;
+#if WITH_EDITOR
+	if (HHitProxy* HitProxy = MeshToolManager->GetContextQueriesAPI()->GetHitProxy(ClickPos.ScreenPosition.X, ClickPos.ScreenPosition.Y))
 	{
-		AActor* Actor = *It;
-		if (Actor->IsEditable() &&
-			Actor->IsListedInSceneOutliner() &&					// Only add actors that are allowed to be selected and drawn in editor
-			!Actor->IsTemplate() &&								// Should never happen, but we never want CDOs
-			!Actor->HasAnyFlags(RF_Transient))				// Don't add transient actors in non-play worlds
+		if (HitProxy->IsA(HActor::StaticGetType()))
 		{
+			HActor* ActorProxy = (HActor*)HitProxy;
+			AActor* Actor = ActorProxy->Actor;
 			TArray<UActorComponent*> CandidateComponents = Actor->K2_GetComponentsByClass(UMeshComponent::StaticClass());
 			for (UActorComponent* CandidateComponent : CandidateComponents)
 			{
@@ -139,17 +138,16 @@ FHitResult UMeshClickTool::FindInitialHitResult(const FInputDeviceRay& ClickPos,
 				TSharedPtr<IMeshPaintComponentAdapter> MeshAdapter = FMeshPaintComponentAdapterFactory::CreateAdapterForMesh(MeshComponent, 0);
 				if (MeshAdapter.IsValid() && IsMeshAdapterSupported(MeshAdapter))
 				{
-					MeshToolManager->AddPaintableMeshComponent(MeshComponent);
 					MeshToolManager->AddToComponentToAdapterMap(MeshComponent, MeshAdapter);
+					CachedClickedComponents.AddUnique(MeshComponent);
+					CachedClickedActors.AddUnique(Cast<AActor>(MeshComponent->GetOuter()));
+					bFoundValidComponents = true;
 				}
 			}
-
 		}
 	}
-	FHitResult TraceHitResult(1.0f);
-	MeshToolManager->FindHitResult(Ray, TraceHitResult);
-	MeshToolManager->ResetState();
-	return TraceHitResult;
+#endif
+	return bFoundValidComponents;
 }
 
 UVertexAdapterClickTool::UVertexAdapterClickTool()
diff --git a/Engine/Plugins/MeshPainting/Source/MeshPaintingToolset/Public/MeshSelect.h b/Engine/Plugins/MeshPainting/Source/MeshPaintingToolset/Public/MeshSelect.h
--- a/Engine/Plugins/MeshPainting/Source/MeshPaintingToolset/Public/MeshSelect.h
+++ b/Engine/Plugins/MeshPainting/Source/MeshPaintingToolset/Public/MeshSelect.h
@@ -7,6 +7,7 @@
 #include "MeshSelect.generated.h"
 
 class IMeshPaintComponentAdapter;
+class AActor;
 
 /**
  * Builder for UVertexAdapterClickTool
@@ -66,15 +67,22 @@ public:
 	{
 		return true;
 	}
+	virtual bool AllowsMultiselect()
+	{
+		return true;
+	}
 
 protected:
 	// flags used to identify modifier keys/buttons
 	static const int AdditiveSelectionModifier = 1;
 	bool bAddToSelectionSet;
 
+	UPROPERTY(Transient)
+	TArray<UMeshComponent*> CachedClickedComponents;
+	UPROPERTY(Transient)
+	TArray<AActor*> CachedClickedActors;
 private:
-	FHitResult FindInitialHitResult(const FInputDeviceRay& ClickPos, class UMeshToolManager* MeshToolManager);
-
+	bool FindClickedComponentsAndCacheAdapters(const FInputDeviceRay& ClickPos, class UMeshToolManager* MeshToolManager);
 };
 
 UCLASS()
@@ -100,5 +108,9 @@ public:
 
 	// USingleClickTool overrides
 	virtual bool IsMeshAdapterSupported(TSharedPtr<IMeshPaintComponentAdapter> MeshAdapter) override;
+	virtual bool AllowsMultiselect()
+	{
+		return false;
+	}
 
 };
diff --git a/Engine/Source/Editor/Experimental/EditorInteractiveToolsFramework/Private/EdModeInteractiveToolsContext.cpp b/Engine/Source/Editor/Experimental/EditorInteractiveToolsFramework/Private/EdModeInteractiveToolsContext.cpp
--- a/Engine/Source/Editor/Experimental/EditorInteractiveToolsFramework/Private/EdModeInteractiveToolsContext.cpp
+++ b/Engine/Source/Editor/Experimental/EditorInteractiveToolsFramework/Private/EdModeInteractiveToolsContext.cpp
@@ -29,6 +29,7 @@
 
 //#define ENABLE_DEBUG_PRINTING
 
+class HHitProxy;
 
 
 static float SnapToIncrement(float fValue, float fIncrement, float offset = 0)
@@ -300,6 +301,8 @@ public:
 		return nullptr;
 	}
 
+	virtual HHitProxy* GetHitProxy(int32 X, int32 Y) const { return nullptr;  }
+
 };
 
 
diff --git a/Engine/Source/Editor/UnrealEd/Private/Tools/UEdMode.cpp b/Engine/Source/Editor/UnrealEd/Private/Tools/UEdMode.cpp
--- a/Engine/Source/Editor/UnrealEd/Private/Tools/UEdMode.cpp
+++ b/Engine/Source/Editor/UnrealEd/Private/Tools/UEdMode.cpp
@@ -33,7 +33,7 @@ public:
 	UInteractiveToolsContext* ToolsContext;
 	UEdMode* EdMode;
 	FViewCameraState CachedViewState;
-
+	FEditorViewportClient* CachedViewportClient;
 	FEdModeQueriesImpl(UInteractiveToolsContext* Context, UEdMode* InEdMode)
 	{
 		ToolsContext = Context;
@@ -42,6 +42,7 @@ public:
 
 	void CacheCurrentViewState(FEditorViewportClient* ViewportClient)
 	{
+		CachedViewportClient = ViewportClient;
 		FViewportCameraTransform ViewTransform = ViewportClient->GetViewTransform();
 		CachedViewState.bIsOrthographic = ViewportClient->IsOrtho();
 		CachedViewState.Position = ViewTransform.GetLocation();
@@ -266,6 +267,15 @@ public:
 		return nullptr;
 	}
 
+	virtual HHitProxy* GetHitProxy(int32 X, int32 Y) const
+	{
+		if (CachedViewportClient && CachedViewportClient->Viewport)
+		{
+			return CachedViewportClient->Viewport->GetHitProxy(X, Y);
+		}
+		return nullptr;
+	}
+
 };
 
 
diff --git a/Engine/Source/Runtime/Experimental/InteractiveToolsFramework/Public/ToolContextInterfaces.h b/Engine/Source/Runtime/Experimental/InteractiveToolsFramework/Public/ToolContextInterfaces.h
--- a/Engine/Source/Runtime/Experimental/InteractiveToolsFramework/Public/ToolContextInterfaces.h
+++ b/Engine/Source/Runtime/Experimental/InteractiveToolsFramework/Public/ToolContextInterfaces.h
@@ -17,7 +17,9 @@ class UInteractiveToolManager;
 class UInteractiveGizmoManager;
 struct FMeshDescription;
 
-
+#if WITH_EDITOR
+class HHitProxy;
+#endif
 
 /**
  * FToolBuilderState is a bucket of state information that a ToolBuilder might need
@@ -202,6 +204,13 @@ public:
 	 * @return Instance of material to use for this purpose
 	 */
 	virtual UMaterialInterface* GetStandardMaterial(EStandardToolContextMaterials MaterialType) const = 0;
+
+#if WITH_EDITOR
+	/**
+	* When selecting, sometimes we need a hit proxy rather than a physics trace or other raycast.                                                                   
+	*/
+	virtual HHitProxy* GetHitProxy(int32 X, int32 Y) const = 0;
+#endif
 };
 
 
