Fix for incremental build

diff --git a/Engine/Plugins/MeshPainting/Source/MeshPaintingToolset/Public/MeshSelect.h b/Engine/Plugins/MeshPainting/Source/MeshPaintingToolset/Public/MeshSelect.h
--- a/Engine/Plugins/MeshPainting/Source/MeshPaintingToolset/Public/MeshSelect.h
+++ b/Engine/Plugins/MeshPainting/Source/MeshPaintingToolset/Public/MeshSelect.h
@@ -8,6 +8,7 @@
 
 class IMeshPaintComponentAdapter;
 class AActor;
+class UMeshComponent;
 
 /**
  * Builder for UVertexAdapterClickTool
