Fix for a crash when leaving mesh paint while interacting with the UI

diff --git a/Engine/Plugins/MeshPainting/Source/MeshPaintEditorMode/Private/BrushSettingsCustomization.cpp b/Engine/Plugins/MeshPainting/Source/MeshPaintEditorMode/Private/BrushSettingsCustomization.cpp
--- a/Engine/Plugins/MeshPainting/Source/MeshPaintEditorMode/Private/BrushSettingsCustomization.cpp
+++ b/Engine/Plugins/MeshPainting/Source/MeshPaintEditorMode/Private/BrushSettingsCustomization.cpp
@@ -245,7 +245,7 @@ void FColorPaintingSettingsCustomization::CustomizeDetails(IDetailLayoutBuilder&
 		.AutoWidth()
 		[
 			SNew(SCheckBox)
-				.IsChecked_Lambda([=]() -> ECheckBoxState { return (UMeshPaintMode::GetColorToolProperties()->bPaintOnSpecificLOD ? ECheckBoxState::Checked : ECheckBoxState::Unchecked); })
+				.IsChecked_Lambda([=]() -> ECheckBoxState { return (UMeshPaintMode::GetColorToolProperties() && UMeshPaintMode::GetColorToolProperties()->bPaintOnSpecificLOD ? ECheckBoxState::Checked : ECheckBoxState::Unchecked); })
 				.OnCheckStateChanged(FOnCheckStateChanged::CreateLambda([=](ECheckBoxState State) { 
 					if (UMeshColorPaintingTool* ColorBrush = Cast<UMeshColorPaintingTool>(UMeshPaintMode::GetMeshPaintMode()->GetToolManager()->GetActiveTool(EToolSide::Left)))
 					{
@@ -257,9 +257,9 @@ void FColorPaintingSettingsCustomization::CustomizeDetails(IDetailLayoutBuilder&
 		.Padding(0.0f, 0.0f, 4.0f, 0.0f)
 		[
 			SNew(SNumericEntryBox<int32>)
-			.IsEnabled_Lambda([=]() -> bool { return UMeshPaintMode::GetColorToolProperties()->bPaintOnSpecificLOD;  })
+			.IsEnabled_Lambda([=]() -> bool { return UMeshPaintMode::GetColorToolProperties() ? UMeshPaintMode::GetColorToolProperties()->bPaintOnSpecificLOD : false;  })
 			.AllowSpin(true)
-			.Value_Lambda([=]() -> int32 { return UMeshPaintMode::GetColorToolProperties()->LODIndex; })
+			.Value_Lambda([=]() -> int32 { return UMeshPaintMode::GetColorToolProperties() ? UMeshPaintMode::GetColorToolProperties()->LODIndex : 0; })
 			.MinValue(0)
 			.MaxValue_Lambda([=]() -> int32 { 
 					if (UMeshColorPaintingTool* ColorBrush = Cast<UMeshColorPaintingTool>(UMeshPaintMode::GetMeshPaintMode()->GetToolManager()->GetActiveTool(EToolSide::Left)))
