Adds more validity checks in case an adapter is being deleted while you are clicking

diff --git a/Engine/Plugins/MeshPainting/Source/MeshPaintingToolset/Private/MeshSelect.cpp b/Engine/Plugins/MeshPainting/Source/MeshPaintingToolset/Private/MeshSelect.cpp
--- a/Engine/Plugins/MeshPainting/Source/MeshPaintingToolset/Private/MeshSelect.cpp
+++ b/Engine/Plugins/MeshPainting/Source/MeshPaintingToolset/Private/MeshSelect.cpp
@@ -158,7 +158,7 @@ UVertexAdapterClickTool::UVertexAdapterClickTool()
 
 bool UVertexAdapterClickTool::IsMeshAdapterSupported(TSharedPtr<IMeshPaintComponentAdapter> MeshAdapter)
 {
-	return MeshAdapter->SupportsVertexPaint();
+	return MeshAdapter.IsValid() ? MeshAdapter->SupportsVertexPaint() : false;
 }
 
 UTextureAdapterClickTool::UTextureAdapterClickTool()
@@ -169,7 +169,7 @@ UTextureAdapterClickTool::UTextureAdapterClickTool()
 
 bool UTextureAdapterClickTool::IsMeshAdapterSupported(TSharedPtr<IMeshPaintComponentAdapter> MeshAdapter)
 {
-	return MeshAdapter->SupportsTexturePaint();
+	return MeshAdapter.IsValid() ? MeshAdapter->SupportsTexturePaint() : false;
 }
 
 #undef LOCTEXT_NAMESPACE
\ No newline at end of file
