Making the bound delegate a shared pointer association instead of a raw one so it is safe on force delete.

diff --git a/Engine/Plugins/MeshPainting/Source/MeshPaintingToolset/Private/MeshPaintSkeletalMeshAdapter.cpp b/Engine/Plugins/MeshPainting/Source/MeshPaintingToolset/Private/MeshPaintSkeletalMeshAdapter.cpp
--- a/Engine/Plugins/MeshPainting/Source/MeshPaintingToolset/Private/MeshPaintSkeletalMeshAdapter.cpp
+++ b/Engine/Plugins/MeshPainting/Source/MeshPaintingToolset/Private/MeshPaintSkeletalMeshAdapter.cpp
@@ -18,7 +18,7 @@ bool FMeshPaintSkeletalMeshComponentAdapter::Construct(UMeshComponent* InCompone
 	SkeletalMeshComponent = Cast<USkeletalMeshComponent>(InComponent);
 	if (SkeletalMeshComponent != nullptr)
 	{
-		SkeletalMeshChangedHandle = SkeletalMeshComponent->RegisterOnSkeletalMeshPropertyChanged(USkeletalMeshComponent::FOnSkeletalMeshPropertyChanged::CreateRaw(this, &FMeshPaintSkeletalMeshComponentAdapter::OnSkeletalMeshChanged));
+		SkeletalMeshChangedHandle = SkeletalMeshComponent->RegisterOnSkeletalMeshPropertyChanged(USkeletalMeshComponent::FOnSkeletalMeshPropertyChanged::CreateSP(this, &FMeshPaintSkeletalMeshComponentAdapter::OnSkeletalMeshChanged));
 
 		if (SkeletalMeshComponent->SkeletalMesh != nullptr)
 		{
diff --git a/Engine/Plugins/MeshPainting/Source/MeshPaintingToolset/Public/BaseMeshPaintComponentAdapter.h b/Engine/Plugins/MeshPainting/Source/MeshPaintingToolset/Public/BaseMeshPaintComponentAdapter.h
--- a/Engine/Plugins/MeshPainting/Source/MeshPaintingToolset/Public/BaseMeshPaintComponentAdapter.h
+++ b/Engine/Plugins/MeshPainting/Source/MeshPaintingToolset/Public/BaseMeshPaintComponentAdapter.h
@@ -12,7 +12,7 @@ struct FIndex3i;
 typedef TIndexMeshArrayAdapter<uint32, double> FIndexMeshArrayAdapterd;
 
 /** Base mesh paint geometry adapter, handles basic sphere intersection using a Octree */
-class MESHPAINTINGTOOLSET_API FBaseMeshPaintComponentAdapter : public IMeshPaintComponentAdapter, public FGCObject
+class MESHPAINTINGTOOLSET_API FBaseMeshPaintComponentAdapter : public IMeshPaintComponentAdapter, public FGCObject, public TSharedFromThis<FBaseMeshPaintComponentAdapter>
 {
 public:
 	/** Start IMeshPaintGeometryAdapter Overrides */
