Remove legacy mesh paint mode.

diff --git a/Engine/Plugins/MeshPainting/Source/MeshPaintEditorMode/Private/MeshPaintEditorMode.cpp b/Engine/Plugins/MeshPainting/Source/MeshPaintEditorMode/Private/MeshPaintEditorMode.cpp
--- a/Engine/Plugins/MeshPainting/Source/MeshPaintEditorMode/Private/MeshPaintEditorMode.cpp
+++ b/Engine/Plugins/MeshPainting/Source/MeshPaintEditorMode/Private/MeshPaintEditorMode.cpp
@@ -1,7 +1,6 @@
 // Copyright Epic Games, Inc. All Rights Reserved.
 
 #include "MeshPaintEditorMode.h"
-#include "EditorModeRegistry.h"
 #include "Modules/ModuleManager.h"
 #include "MeshPaintMode.h"
 #include "EditorModeManager.h"
@@ -14,7 +13,6 @@
 #include "MeshPaintSplineMeshAdapter.h"
 #include "MeshPaintStaticMeshAdapter.h"
 #include "MeshPaintSkeletalMeshAdapter.h"
-#include "Settings/LevelEditorMiscSettings.h"
 
 class FMeshPaintEditorModeModule : public IModuleInterface
 {
@@ -57,12 +55,9 @@ void FMeshPaintEditorModeModule::Register()
 
 void FMeshPaintEditorModeModule::OnPostEngineInit()
 {
-	if (!GetDefault<ULevelEditorMiscSettings>()->bEnableLegacyMeshPaintMode)
-	{
-		Register();
-		FMeshPaintingToolActionCommands::RegisterAllToolActions();
-		FMeshPaintEditorModeCommands::Register();
-	}
+	Register();
+	FMeshPaintingToolActionCommands::RegisterAllToolActions();
+	FMeshPaintEditorModeCommands::Register();
 }
 
 void FMeshPaintEditorModeModule::StartupModule()
diff --git a/Engine/Plugins/MeshPainting/Source/MeshPaintEditorMode/Private/MeshPaintMode.cpp b/Engine/Plugins/MeshPainting/Source/MeshPaintEditorMode/Private/MeshPaintMode.cpp
--- a/Engine/Plugins/MeshPainting/Source/MeshPaintEditorMode/Private/MeshPaintMode.cpp
+++ b/Engine/Plugins/MeshPainting/Source/MeshPaintEditorMode/Private/MeshPaintMode.cpp
@@ -118,15 +118,13 @@ UMeshPaintMode::UMeshPaintMode()
 	: Super()
 {
 	SettingsClass = UMeshPaintModeSettings::StaticClass();
-	// Don't be a visible mode unless legacy mesh paint mode is not on.
-	const bool bVisible = !GetDefault<ULevelEditorMiscSettings>()->bEnableLegacyMeshPaintMode;
 	FModuleManager::Get().LoadModule("EditorStyle");
 
 	Info = FEditorModeInfo(
 		FName(TEXT("MeshPaintMode")),
 		LOCTEXT("ModeName", "Mesh Paint"),
 		FSlateIcon(FEditorStyle::GetStyleSetName(), "LevelEditor.MeshPaintMode", "LevelEditor.MeshPaintMode.Small"),
-		bVisible,
+		true,
 		600
 	);
 }
diff --git a/Engine/Source/Editor/MeshPaintMode/MeshPaintMode.Build.cs b/Engine/Source/Editor/MeshPaintMode/MeshPaintMode.Build.cs
deleted file mode 100644
diff --git a/Engine/Source/Editor/MeshPaintMode/Private/MeshPaintEdMode.cpp b/Engine/Source/Editor/MeshPaintMode/Private/MeshPaintEdMode.cpp
deleted file mode 100644
diff --git a/Engine/Source/Editor/MeshPaintMode/Private/MeshPaintModeModule.cpp b/Engine/Source/Editor/MeshPaintMode/Private/MeshPaintModeModule.cpp
deleted file mode 100644
diff --git a/Engine/Source/Editor/MeshPaintMode/Private/MeshPaintModeToolKit.cpp b/Engine/Source/Editor/MeshPaintMode/Private/MeshPaintModeToolKit.cpp
deleted file mode 100644
diff --git a/Engine/Source/Editor/MeshPaintMode/Private/MeshPaintModeToolKit.h b/Engine/Source/Editor/MeshPaintMode/Private/MeshPaintModeToolKit.h
deleted file mode 100644
diff --git a/Engine/Source/Editor/MeshPaintMode/Private/PaintModeCommands.cpp b/Engine/Source/Editor/MeshPaintMode/Private/PaintModeCommands.cpp
deleted file mode 100644
diff --git a/Engine/Source/Editor/MeshPaintMode/Private/PaintModeCommands.h b/Engine/Source/Editor/MeshPaintMode/Private/PaintModeCommands.h
deleted file mode 100644
diff --git a/Engine/Source/Editor/MeshPaintMode/Private/PaintModePainter.cpp b/Engine/Source/Editor/MeshPaintMode/Private/PaintModePainter.cpp
deleted file mode 100644
diff --git a/Engine/Source/Editor/MeshPaintMode/Private/PaintModePainter.h b/Engine/Source/Editor/MeshPaintMode/Private/PaintModePainter.h
deleted file mode 100644
diff --git a/Engine/Source/Editor/MeshPaintMode/Private/PaintModeSettings.cpp b/Engine/Source/Editor/MeshPaintMode/Private/PaintModeSettings.cpp
deleted file mode 100644
diff --git a/Engine/Source/Editor/MeshPaintMode/Private/PaintModeSettings.h b/Engine/Source/Editor/MeshPaintMode/Private/PaintModeSettings.h
deleted file mode 100644
diff --git a/Engine/Source/Editor/MeshPaintMode/Private/PaintModeSettingsCustomization.cpp b/Engine/Source/Editor/MeshPaintMode/Private/PaintModeSettingsCustomization.cpp
deleted file mode 100644
diff --git a/Engine/Source/Editor/MeshPaintMode/Private/PaintModeSettingsCustomization.h b/Engine/Source/Editor/MeshPaintMode/Private/PaintModeSettingsCustomization.h
deleted file mode 100644
diff --git a/Engine/Source/Editor/MeshPaintMode/Private/SPaintModeWidget.cpp b/Engine/Source/Editor/MeshPaintMode/Private/SPaintModeWidget.cpp
deleted file mode 100644
diff --git a/Engine/Source/Editor/MeshPaintMode/Private/SPaintModeWidget.h b/Engine/Source/Editor/MeshPaintMode/Private/SPaintModeWidget.h
deleted file mode 100644
diff --git a/Engine/Source/Editor/MeshPaintMode/Public/MeshPaintEdMode.h b/Engine/Source/Editor/MeshPaintMode/Public/MeshPaintEdMode.h
deleted file mode 100644
diff --git a/Engine/Source/Editor/MeshPaintMode/Public/MeshPaintModeModule.h b/Engine/Source/Editor/MeshPaintMode/Public/MeshPaintModeModule.h
deleted file mode 100644
diff --git a/Engine/Source/Editor/UnrealEd/Classes/Settings/LevelEditorMiscSettings.h b/Engine/Source/Editor/UnrealEd/Classes/Settings/LevelEditorMiscSettings.h
--- a/Engine/Source/Editor/UnrealEd/Classes/Settings/LevelEditorMiscSettings.h
+++ b/Engine/Source/Editor/UnrealEd/Classes/Settings/LevelEditorMiscSettings.h
@@ -43,11 +43,6 @@ public:
 	UPROPERTY(EditAnywhere, config, Category=Editing, meta=( DisplayName = "Preserve Actor Scale on Replace" ))
 	uint32 bReplaceRespectsScale:1;
 
-
-	/** If enabled, the modes tab will revert to the original mesh paint */
-	UPROPERTY(EditAnywhere, config, AdvancedDisplay, Category = Editing, meta = (ConfigRestartRequired = true))
-	uint32 bEnableLegacyMeshPaintMode : 1;
-
 	/** If enabled, will avoid relabeling actors in UUnrealEdEngine::edactPasteSelected */
 	UPROPERTY(EditAnywhere, config, Category = Editing, meta = (DisplayName = "Avoid Actor Relabel on Paste Selected"))
 	uint32 bAvoidRelabelOnPasteSelected:1;
diff --git a/Engine/Source/Editor/UnrealEd/Private/EditorModeRegistry.cpp b/Engine/Source/Editor/UnrealEd/Private/EditorModeRegistry.cpp
--- a/Engine/Source/Editor/UnrealEd/Private/EditorModeRegistry.cpp
+++ b/Engine/Source/Editor/UnrealEd/Private/EditorModeRegistry.cpp
@@ -60,7 +60,6 @@ void FEditorModeRegistry::Initialize()
 
 	FModuleManager::LoadModuleChecked<FActorPickerModeModule>(TEXT("ActorPickerMode"));
 	FModuleManager::LoadModuleChecked<FSceneDepthPickerModeModule>(TEXT("SceneDepthPickerMode"));
-	FModuleManager::LoadModuleChecked<IMeshPaintModule>(TEXT("MeshPaintMode"));
 	FModuleManager::LoadModuleChecked<ILandscapeEditorModule>(TEXT("LandscapeEditor"));
 	FModuleManager::LoadModuleChecked<IFoliageEditModule>(TEXT("FoliageEdit"));
 	FModuleManager::LoadModuleChecked<IVirtualTexturingEditorModule>(TEXT("VirtualTexturingEditor"));
diff --git a/Engine/Source/Editor/UnrealEd/Private/Settings/SettingsClasses.cpp b/Engine/Source/Editor/UnrealEd/Private/Settings/SettingsClasses.cpp
--- a/Engine/Source/Editor/UnrealEd/Private/Settings/SettingsClasses.cpp
+++ b/Engine/Source/Editor/UnrealEd/Private/Settings/SettingsClasses.cpp
@@ -425,7 +425,6 @@ ULevelEditorMiscSettings::ULevelEditorMiscSettings( const FObjectInitializer& Ob
 	PercentageThresholdForPrompt = 20.0f;
 	MinimumBoundsForCheckingSize = FVector(500.0f, 500.0f, 50.0f);
 	bCreateNewAudioDeviceForPlayInEditor = true;
-	bEnableLegacyMeshPaintMode = false;
 	bAvoidRelabelOnPasteSelected = false;
 }
 
diff --git a/Engine/Source/Editor/UnrealEd/UnrealEd.Build.cs b/Engine/Source/Editor/UnrealEd/UnrealEd.Build.cs
--- a/Engine/Source/Editor/UnrealEd/UnrealEd.Build.cs
+++ b/Engine/Source/Editor/UnrealEd/UnrealEd.Build.cs
@@ -157,7 +157,6 @@ public class UnrealEd : ModuleRules
 				"JsonUtilities",
 				"Landscape",
 				"MeshPaint",
-				"MeshPaintMode",
 				"Foliage",
 				"VectorVM",
 				"MaterialUtilities",
@@ -280,7 +279,6 @@ public class UnrealEd : ModuleRules
 				"ViewportInteraction",
 				"VREditor",
 				"MeshPaint",
-				"MeshPaintMode",
 				"PropertyEditor",
 				"ToolMenusEditor",
 				"InputBindingEditor",
diff --git a/Engine/Source/Programs/UnrealBuildTool/Configuration/UEBuildModuleCPP.cs b/Engine/Source/Programs/UnrealBuildTool/Configuration/UEBuildModuleCPP.cs
--- a/Engine/Source/Programs/UnrealBuildTool/Configuration/UEBuildModuleCPP.cs
+++ b/Engine/Source/Programs/UnrealBuildTool/Configuration/UEBuildModuleCPP.cs
@@ -166,7 +166,6 @@ namespace UnrealBuildTool
 			new KeyValuePair<string, string>("WebBrowser", "WebBrowserTexture"),
 			new KeyValuePair<string, string>("WindowsMixedRealityHMD", "WindowsMixedRealityHandTracking"),
             new KeyValuePair<string, string>("UnrealEd", "MeshPaint"),
-            new KeyValuePair<string, string>("UnrealEd", "MeshPaintMode"),
         };
 
 
