Fixing a crash when swapping selected actors in texture paint mode.

# This is related to the change to shared mesh tool data from UInteractiveToolManger. It's not valid to do this cast, and I missed it in the original change.

#rnx

diff --git a/Engine/Plugins/MeshPainting/Source/MeshPaintingToolset/Private/MeshTexturePaintingTool.cpp b/Engine/Plugins/MeshPainting/Source/MeshPaintingToolset/Private/MeshTexturePaintingTool.cpp
--- a/Engine/Plugins/MeshPainting/Source/MeshPaintingToolset/Private/MeshTexturePaintingTool.cpp
+++ b/Engine/Plugins/MeshPainting/Source/MeshPaintingToolset/Private/MeshTexturePaintingTool.cpp
@@ -310,8 +310,6 @@ bool UMeshTexturePaintingTool::PaintInternal(const TArrayView<TPair<FVector, FVe
 	bool bPaintApplied = false;
 
 	// Fire out a ray to see if there is a *selected* component under the mouse cursor that can be painted.
-	UMeshToolManager* MeshToolManager = Cast<UMeshToolManager>(GetToolManager());
-
 	for (int32 i = 0; i < Rays.Num(); ++i)
 	{
 		const FVector& RayOrigin = Rays[i].Key;
@@ -321,19 +319,22 @@ bool UMeshTexturePaintingTool::PaintInternal(const TArrayView<TPair<FVector, FVe
 		const FVector TraceStart(RayOrigin);
 		const FVector TraceEnd(RayOrigin + RayDirection * HALF_WORLD_MAX);
 
-		for (UMeshComponent* MeshComponent : MeshToolManager->GetPaintableMeshComponents())
+		if (SharedMeshToolData.IsValid())
 		{
-			TSharedPtr<IMeshPaintComponentAdapter> MeshAdapter = MeshToolManager->GetAdapterForComponent(MeshComponent);
+			for (UMeshComponent* MeshComponent : SharedMeshToolData->GetPaintableMeshComponents())
+			{
+				TSharedPtr<IMeshPaintComponentAdapter> MeshAdapter = SharedMeshToolData->GetAdapterForComponent(MeshComponent);
 
-			// Ray trace
-			FHitResult TraceHitResult(1.0f);
+				// Ray trace
+				FHitResult TraceHitResult(1.0f);
 
-			if (MeshAdapter->LineTraceComponent(TraceHitResult, TraceStart, TraceEnd, FCollisionQueryParams(SCENE_QUERY_STAT(Paint), true)))
-			{
-				// Find the closest impact
-				if ((BestTraceResult.GetComponent() == nullptr) || (TraceHitResult.Time < BestTraceResult.Time))
+				if (MeshAdapter->LineTraceComponent(TraceHitResult, TraceStart, TraceEnd, FCollisionQueryParams(SCENE_QUERY_STAT(Paint), true)))
 				{
-					BestTraceResult = TraceHitResult;
+					// Find the closest impact
+					if ((BestTraceResult.GetComponent() == nullptr) || (TraceHitResult.Time < BestTraceResult.Time))
+					{
+						BestTraceResult = TraceHitResult;
+					}
 				}
 			}
 		}
@@ -408,7 +409,7 @@ bool UMeshTexturePaintingTool::PaintInternal(const TArrayView<TPair<FVector, FVe
 			UMeshComponent* HoveredComponent = Entry.Key;
 			TArray<int32>& PaintRayResultIds = Entry.Value;
 
-			IMeshPaintComponentAdapter* MeshAdapter = MeshToolManager->GetAdapterForComponent(HoveredComponent).Get();
+			IMeshPaintComponentAdapter* MeshAdapter = SharedMeshToolData.IsValid() ? SharedMeshToolData->GetAdapterForComponent(HoveredComponent).Get() : nullptr;
 			if (!ensure(MeshAdapter))
 			{
 				continue;
