To more closely match the recommended setup of using subsystems to handle operations that are called from editor mode UI, this is a refactor of the experimental mesh paint plugin Blueprint libraries to use subsystems instead.

diff --git a/Engine/Plugins/MeshPainting/Source/MeshPaintEditorMode/MeshPaintEditorMode.Build.cs b/Engine/Plugins/MeshPainting/Source/MeshPaintEditorMode/MeshPaintEditorMode.Build.cs
--- a/Engine/Plugins/MeshPainting/Source/MeshPaintEditorMode/MeshPaintEditorMode.Build.cs
+++ b/Engine/Plugins/MeshPainting/Source/MeshPaintEditorMode/MeshPaintEditorMode.Build.cs
@@ -12,6 +12,7 @@ namespace UnrealBuildTool.Rules
                     "Slate",
                     "SlateCore",
                     "EditorStyle",
+					"EditorSubsystem"
                 }
             );
 
diff --git a/Engine/Plugins/MeshPainting/Source/MeshPaintEditorMode/Private/BrushSettingsCustomization.cpp b/Engine/Plugins/MeshPainting/Source/MeshPaintEditorMode/Private/BrushSettingsCustomization.cpp
--- a/Engine/Plugins/MeshPainting/Source/MeshPaintEditorMode/Private/BrushSettingsCustomization.cpp
+++ b/Engine/Plugins/MeshPainting/Source/MeshPaintEditorMode/Private/BrushSettingsCustomization.cpp
@@ -162,7 +162,7 @@ FReply FVertexPaintingSettingsCustomization::OnSwapColorsClicked(TSharedRef<IPro
 {
 	FScopedTransaction Transaction(NSLOCTEXT("VertexPaintSettings", "SwapColorsTransation", "Swap paint and erase colors"));
 
-	UMeshPaintModeHelpers::SwapVertexColors();
+	GEditor->GetEditorSubsystem<UMeshPaintModeSubsystem>()->SwapVertexColors();
 	UMeshVertexPaintingToolProperties* Settings = UMeshPaintMode::GetVertexToolProperties();
 	if (Settings)
 	{
diff --git a/Engine/Plugins/MeshPainting/Source/MeshPaintEditorMode/Private/ImportVertexColorOptions.cpp b/Engine/Plugins/MeshPainting/Source/MeshPaintEditorMode/Private/ImportVertexColorOptions.cpp
--- a/Engine/Plugins/MeshPainting/Source/MeshPaintEditorMode/Private/ImportVertexColorOptions.cpp
+++ b/Engine/Plugins/MeshPainting/Source/MeshPaintEditorMode/Private/ImportVertexColorOptions.cpp
@@ -23,10 +23,10 @@ void SImportVertexColorOptionsWindow::Construct(const FArguments& InArgs)
 	Options->UVIndex = 0;
 	
 	// Populate max UV index for each LOD in the mesh component
-	const int32 NumLODs = UMeshPaintingToolset::GetNumberOfLODs(InArgs._Component);
+	const int32 NumLODs = GEngine->GetEngineSubsystem<UMeshPaintingSubsystem>()->GetNumberOfLODs(InArgs._Component);
 	for (int32 LODIndex = 0; LODIndex < NumLODs; ++LODIndex)
 	{
-		Options->LODToMaxUVMap.Add(LODIndex, UMeshPaintingToolset::GetNumberOfUVs(InArgs._Component, LODIndex) - 1 );
+		Options->LODToMaxUVMap.Add(LODIndex, GEngine->GetEngineSubsystem<UMeshPaintingSubsystem>()->GetNumberOfUVs(InArgs._Component, LODIndex) - 1 );
 	}
 	Options->NumLODs = NumLODs;
 	// Can only import vertex colors to static mesh component instances
diff --git a/Engine/Plugins/MeshPainting/Source/MeshPaintEditorMode/Private/MeshPaintMode.cpp b/Engine/Plugins/MeshPainting/Source/MeshPaintEditorMode/Private/MeshPaintMode.cpp
--- a/Engine/Plugins/MeshPainting/Source/MeshPaintEditorMode/Private/MeshPaintMode.cpp
+++ b/Engine/Plugins/MeshPainting/Source/MeshPaintEditorMode/Private/MeshPaintMode.cpp
@@ -211,10 +211,10 @@ void UMeshPaintMode::Tick(FEditorViewportClient* ViewportClient, float DeltaTime
 	if (ViewportClient->IsPerspective())
 	{
 		// Make sure perspective viewports are still set to real-time
-		UMeshPaintModeHelpers::SetRealtimeViewport(true);
+		GEditor->GetEditorSubsystem<UMeshPaintModeSubsystem>()->SetRealtimeViewport(true);
 
 		// Set viewport show flags		
-		UMeshPaintModeHelpers::SetViewportColorMode(ModeSettings->ColorViewMode, ViewportClient);
+		GEditor->GetEditorSubsystem<UMeshPaintModeSubsystem>()->SetViewportColorMode(ModeSettings->ColorViewMode, ViewportClient);
 	}
 
 
@@ -258,7 +258,7 @@ void UMeshPaintMode::BindCommands()
 
 	CommandList->MapAction(Commands.SwitchForeAndBackgroundColor, FExecuteAction::CreateLambda([this]()
 	{
-		UMeshPaintModeHelpers::SwapVertexColors();
+		GEditor->GetEditorSubsystem<UMeshPaintModeSubsystem>()->SwapVertexColors();
 	}));
 
 
@@ -333,7 +333,7 @@ void UMeshPaintMode::BindCommands()
 			FExecuteAction::CreateUObject(this, &UMeshPaintMode::CycleTextures, -1),
 			FCanExecuteAction::CreateUObject(this, &UMeshPaintMode::CanCycleTextures)
 		));
-	CommandList->MapAction(Commands.SaveTexturePaint, FUIAction(FExecuteAction::CreateStatic(&UMeshPaintModeHelpers::SaveModifiedTextures), FCanExecuteAction::CreateStatic(&UMeshPaintModeHelpers::CanSaveModifiedTextures)));
+	CommandList->MapAction(Commands.SaveTexturePaint, FUIAction(FExecuteAction::CreateUObject(GEditor->GetEditorSubsystem<UMeshPaintModeSubsystem>(), &UMeshPaintModeSubsystem::SaveModifiedTextures), FCanExecuteAction::CreateUObject(GEditor->GetEditorSubsystem<UMeshPaintModeSubsystem>(), &UMeshPaintModeSubsystem::CanSaveModifiedTextures)));
 
 	auto HasPaintChanges = [this]() -> bool { return GetNumberOfPendingPaintChanges() > 0; };
 	CommandList->MapAction(Commands.PropagateTexturePaint, FExecuteAction::CreateUObject(this, &UMeshPaintMode::CommitAllPaintedTextures), FCanExecuteAction::CreateLambda(HasPaintChanges));
@@ -404,7 +404,7 @@ void UMeshPaintMode::FillWithVertexColor()
 
 	if (GetToolManager()->GetActiveTool(EToolSide::Mouse)->IsA<UMeshWeightPaintingTool>())
 	{
-		FillColor = UMeshPaintingToolset::GenerateColorForTextureWeight((int32)GetWeightToolProperties()->TextureWeightType, (int32)GetWeightToolProperties()->PaintTextureWeightIndex).ToFColor(bConvertSRGB);
+		FillColor = GEngine->GetEngineSubsystem<UMeshPaintingSubsystem>()->GenerateColorForTextureWeight((int32)GetWeightToolProperties()->TextureWeightType, (int32)GetWeightToolProperties()->PaintTextureWeightIndex).ToFColor(bConvertSRGB);
 	}
 	else if (UMeshColorPaintingToolProperties* ColorProperties = GetColorToolProperties())
 	{
@@ -435,11 +435,11 @@ void UMeshPaintMode::FillWithVertexColor()
 
 		if (Component->IsA<UStaticMeshComponent>())
 		{
-			UMeshPaintingToolset::FillStaticMeshVertexColors(Cast<UStaticMeshComponent>(Component), bPaintOnSpecificLOD ? ColorProperties->LODIndex : -1, FillColor, MaskColor);
+			GEngine->GetEngineSubsystem<UMeshPaintingSubsystem>()->FillStaticMeshVertexColors(Cast<UStaticMeshComponent>(Component), bPaintOnSpecificLOD ? ColorProperties->LODIndex : -1, FillColor, MaskColor);
 		}
 		else if (Component->IsA<USkeletalMeshComponent>())
 		{
-			UMeshPaintingToolset::FillSkeletalMeshVertexColors(Cast<USkeletalMeshComponent>(Component), bPaintOnSpecificLOD ? ColorProperties->LODIndex : -1, FillColor, MaskColor);
+			GEngine->GetEngineSubsystem<UMeshPaintingSubsystem>()->FillSkeletalMeshVertexColors(Cast<USkeletalMeshComponent>(Component), bPaintOnSpecificLOD ? ColorProperties->LODIndex : -1, FillColor, MaskColor);
 		}
 		
 
@@ -468,7 +468,7 @@ void UMeshPaintMode::PropagateVertexColorsToAsset()
 	if (VertexColorCopyWarning.ShowModal() != FSuppressableWarningDialog::Cancel)
 	{
 		FScopedTransaction Transaction(LOCTEXT("LevelMeshPainter_TransactionPropogateColors", "Propagating Vertex Colors To Source Meshes"));
-		UMeshPaintModeHelpers::PropagateVertexColors(StaticMeshComponents);
+		GEditor->GetEditorSubsystem<UMeshPaintModeSubsystem>()->PropagateVertexColors(StaticMeshComponents);
 	}
 }
 
@@ -479,7 +479,7 @@ bool UMeshPaintMode::CanPropagateVertexColors() const
 
 	TArray<UStaticMesh*> StaticMeshes;
 	TArray<UStaticMeshComponent*> StaticMeshComponents = GetSelectedComponents<UStaticMeshComponent>();
-	return UMeshPaintModeHelpers::CanPropagateVertexColors(StaticMeshComponents, StaticMeshes, NumInstanceVertexColorBytes);
+	return GEditor->GetEditorSubsystem<UMeshPaintModeSubsystem>()->CanPropagateVertexColors(StaticMeshComponents, StaticMeshes, NumInstanceVertexColorBytes);
 
 }
 
@@ -490,7 +490,7 @@ void UMeshPaintMode::ImportVertexColors()
 	{
 		/** Import vertex color to single selected mesh component */
 		FScopedTransaction Transaction(LOCTEXT("LevelMeshPainter_TransactionImportColors", "Importing Vertex Colors From Texture"));
-		UMeshPaintModeHelpers::ImportVertexColorsFromTexture(MeshComponents[0]);
+		GEditor->GetEditorSubsystem<UMeshPaintModeSubsystem>()->ImportVertexColorsFromTexture(MeshComponents[0]);
 	}
 }
 
@@ -566,7 +566,7 @@ bool UMeshPaintMode::CanRemoveInstanceColors() const
 	{
 		if (Component != nullptr && Component->GetStaticMesh() != nullptr && Component->GetStaticMesh()->GetNumLODs() > (int32)PaintingMeshLODIndex)
 		{
-			uint32 BufferSize = UMeshPaintingToolset::GetVertexColorBufferSize(Component, PaintingMeshLODIndex, true);
+			uint32 BufferSize = GEngine->GetEngineSubsystem<UMeshPaintingSubsystem>()->GetVertexColorBufferSize(Component, PaintingMeshLODIndex, true);
 
 			if (BufferSize > 0)
 			{
@@ -582,7 +582,7 @@ bool UMeshPaintMode::CanPasteInstanceVertexColors() const
 {
 	const TArray<UStaticMeshComponent*> StaticMeshComponents = GetSelectedComponents<UStaticMeshComponent>();
 	const TArray<FPerComponentVertexColorData> CopiedColorsByComponent = SharedMeshToolData->GetCopiedColorsByComponent();
-	return UMeshPaintModeHelpers::CanPasteInstanceVertexColors(StaticMeshComponents, CopiedColorsByComponent);
+	return GEditor->GetEditorSubsystem<UMeshPaintModeSubsystem>()->CanPasteInstanceVertexColors(StaticMeshComponents, CopiedColorsByComponent);
 }
 
 bool UMeshPaintMode::CanCopyInstanceVertexColors() const
@@ -594,7 +594,7 @@ bool UMeshPaintMode::CanCopyInstanceVertexColors() const
 		PaintingMeshLODIndex = ColorProperties->bPaintOnSpecificLOD ? ColorProperties->LODIndex : 0;
 	}
 
-	return UMeshPaintModeHelpers::CanCopyInstanceVertexColors(StaticMeshComponents, PaintingMeshLODIndex);
+	return GEditor->GetEditorSubsystem<UMeshPaintModeSubsystem>()->CanCopyInstanceVertexColors(StaticMeshComponents, PaintingMeshLODIndex);
 
 }
 
@@ -614,7 +614,7 @@ void UMeshPaintMode::CopyVertexColors()
 {
 	const TArray<UStaticMeshComponent*> StaticMeshComponents = GetSelectedComponents<UStaticMeshComponent>();
 	TArray<FPerComponentVertexColorData> CopiedColorsByComponent;
-	UMeshPaintModeHelpers::CopyVertexColors(StaticMeshComponents, CopiedColorsByComponent);
+	GEditor->GetEditorSubsystem<UMeshPaintModeSubsystem>()->CopyVertexColors(StaticMeshComponents, CopiedColorsByComponent);
 	SharedMeshToolData->SetCopiedColorsByComponent(CopiedColorsByComponent);
 
 }
@@ -624,7 +624,7 @@ void UMeshPaintMode::PasteVertexColors()
 	FScopedTransaction Transaction(LOCTEXT("LevelMeshPainter_TransactionPasteInstColors", "Pasting Per-Instance Vertex Colors"));
 	const TArray<UStaticMeshComponent*> StaticMeshComponents = GetSelectedComponents<UStaticMeshComponent>();
 	TArray<FPerComponentVertexColorData> CopiedColorsByComponent = SharedMeshToolData->GetCopiedColorsByComponent();
-	UMeshPaintModeHelpers::PasteVertexColors(StaticMeshComponents, CopiedColorsByComponent);
+	GEditor->GetEditorSubsystem<UMeshPaintModeSubsystem>()->PasteVertexColors(StaticMeshComponents, CopiedColorsByComponent);
 	UpdateCachedVertexDataSize();
 }
 
@@ -657,7 +657,7 @@ void UMeshPaintMode::RemoveVertexColors()
 	const TArray<UStaticMeshComponent*> StaticMeshComponents = GetSelectedComponents<UStaticMeshComponent>();
 	for (UStaticMeshComponent* Component : StaticMeshComponents)
 	{
-		UMeshPaintingToolset::RemoveComponentInstanceVertexColors(Component);
+		GEngine->GetEngineSubsystem<UMeshPaintingSubsystem>()->RemoveComponentInstanceVertexColors(Component);
 	}
 
 	UpdateCachedVertexDataSize();
@@ -691,7 +691,7 @@ void UMeshPaintMode::PropagateVertexColorsToLODs()
 		{
 			// Reset the state flag as we'll be removing all per-lod colors 
 			SharedMeshToolData->ClearSelectionLODColors();
-			UMeshPaintModeHelpers::RemovePerLODColors(PaintableComponents);
+			GEditor->GetEditorSubsystem<UMeshPaintModeSubsystem>()->RemovePerLODColors(PaintableComponents);
 		}
 	}
 
@@ -706,7 +706,7 @@ void UMeshPaintMode::PropagateVertexColorsToLODs()
 		if (SelectedComponent)
 		{
 			TSharedPtr<IMeshPaintComponentAdapter> MeshAdapter = SharedMeshToolData->GetAdapterForComponent(SelectedComponent);
-			UMeshPaintingToolset::ApplyVertexColorsToAllLODs(*MeshAdapter, SelectedComponent);
+			GEngine->GetEngineSubsystem<UMeshPaintingSubsystem>()->ApplyVertexColorsToAllLODs(*MeshAdapter, SelectedComponent);
 			FComponentReregisterContext ReregisterContext(SelectedComponent);
 		}
 	}
@@ -773,10 +773,10 @@ void UMeshPaintMode::UpdateCachedVertexDataSize()
 	{
 		for (UMeshComponent* SelectedComponent : SharedMeshToolData->GetPaintableMeshComponents())
 		{
-			int32 NumLODs = UMeshPaintingToolset::GetNumberOfLODs(SelectedComponent);
+			int32 NumLODs = GEngine->GetEngineSubsystem<UMeshPaintingSubsystem>()->GetNumberOfLODs(SelectedComponent);
 			for (int32 LODIndex = 0; LODIndex < NumLODs; ++LODIndex)
 			{
-				CachedVertexDataSize += UMeshPaintingToolset::GetVertexColorBufferSize(SelectedComponent, LODIndex, bInstance);
+				CachedVertexDataSize += GEngine->GetEngineSubsystem<UMeshPaintingSubsystem>()->GetVertexColorBufferSize(SelectedComponent, LODIndex, bInstance);
 			}
 		}
 	}
@@ -883,7 +883,7 @@ void UMeshPaintMode::OnResetViewMode()
 			continue;
 		}
 
-		UMeshPaintModeHelpers::SetViewportColorMode(EMeshPaintDataColorViewMode::Normal, ViewportClient);
+		GEditor->GetEditorSubsystem<UMeshPaintModeSubsystem>()->SetViewportColorMode(EMeshPaintDataColorViewMode::Normal, ViewportClient);
 	}
 }
 
diff --git a/Engine/Plugins/MeshPainting/Source/MeshPaintEditorMode/Private/MeshPaintModeHelpers.cpp b/Engine/Plugins/MeshPainting/Source/MeshPaintEditorMode/Private/MeshPaintModeHelpers.cpp
--- a/Engine/Plugins/MeshPainting/Source/MeshPaintEditorMode/Private/MeshPaintModeHelpers.cpp
+++ b/Engine/Plugins/MeshPainting/Source/MeshPaintEditorMode/Private/MeshPaintModeHelpers.cpp
@@ -48,7 +48,7 @@
 #include "MeshTexturePaintingTool.h"
 
 
-void UMeshPaintModeHelpers::SetViewportColorMode(EMeshPaintDataColorViewMode ColorViewMode, FEditorViewportClient* ViewportClient)
+void UMeshPaintModeSubsystem::SetViewportColorMode(EMeshPaintDataColorViewMode ColorViewMode, FEditorViewportClient* ViewportClient)
 {
 	if (ViewportClient->IsPerspective())
 	{
@@ -119,7 +119,7 @@ void UMeshPaintModeHelpers::SetViewportColorMode(EMeshPaintDataColorViewMode Col
 	}
 }
 
-void UMeshPaintModeHelpers::SetRealtimeViewport(bool bRealtime)
+void UMeshPaintModeSubsystem::SetRealtimeViewport(bool bRealtime)
 {
 	FLevelEditorModule& LevelEditorModule = FModuleManager::GetModuleChecked<FLevelEditorModule>("LevelEditor");
 	TSharedPtr< IAssetViewport > ViewportWindow = LevelEditorModule.GetFirstActiveViewport();
@@ -142,7 +142,7 @@ void UMeshPaintModeHelpers::SetRealtimeViewport(bool bRealtime)
 }
 
 
-void UMeshPaintModeHelpers::ImportVertexColorsFromTexture(UMeshComponent* MeshComponent)
+void UMeshPaintModeSubsystem::ImportVertexColorsFromTexture(UMeshComponent* MeshComponent)
 {
 	checkf(MeshComponent != nullptr, TEXT("Invalid mesh component ptr"));
 
@@ -249,7 +249,7 @@ void UMeshPaintModeHelpers::ImportVertexColorsFromTexture(UMeshComponent* MeshCo
 }
 
 
-void UMeshPaintModeHelpers::ImportVertexColorsToSkeletalMesh(USkeletalMesh* SkeletalMesh, const UImportVertexColorOptions* Options, UTexture2D* Texture)
+void UMeshPaintModeSubsystem::ImportVertexColorsToSkeletalMesh(USkeletalMesh* SkeletalMesh, const UImportVertexColorOptions* Options, UTexture2D* Texture)
 {
 	checkf(SkeletalMesh && Options && Texture, TEXT("Invalid ptr"));
 
@@ -281,7 +281,7 @@ void UMeshPaintModeHelpers::ImportVertexColorsToSkeletalMesh(USkeletalMesh* Skel
 		for (uint32 VertexIndex = 0; VertexIndex < LODData.GetNumVertices(); ++VertexIndex)
 		{
 			const FVector2D UV = LODData.StaticVertexBuffers.StaticMeshVertexBuffer.GetVertexUV(VertexIndex, UVIndex);
-			LODData.StaticVertexBuffers.ColorVertexBuffer.VertexColor(VertexIndex) = UMeshPaintingToolset::PickVertexColorFromTextureData(MipData, UV, Texture, ColorMask);
+			LODData.StaticVertexBuffers.ColorVertexBuffer.VertexColor(VertexIndex) = GEngine->GetEngineSubsystem<UMeshPaintingSubsystem>()->PickVertexColorFromTextureData(MipData, UV, Texture, ColorMask);
 		}
 
 		SkeletalMesh->InitResources();
@@ -298,7 +298,7 @@ void UMeshPaintModeHelpers::ImportVertexColorsToSkeletalMesh(USkeletalMesh* Skel
 		LODModel.GetSectionFromVertexIndex(VertexIndex, SectionIndex, SectionVertexIndex);
 
 		const FVector2D UV = LODModel.Sections[SectionIndex].SoftVertices[SectionVertexIndex].UVs[UVIndex];
-		LODModel.Sections[SectionIndex].SoftVertices[SectionVertexIndex].Color = UMeshPaintingToolset::PickVertexColorFromTextureData(MipData, UV, Texture, ColorMask);
+		LODModel.Sections[SectionIndex].SoftVertices[SectionVertexIndex].Color = GEngine->GetEngineSubsystem<UMeshPaintingSubsystem>()->PickVertexColorFromTextureData(MipData, UV, Texture, ColorMask);
 	}
 
 	//Make sure we change the import data so the re-import do not replace the new data
@@ -315,7 +315,7 @@ void UMeshPaintModeHelpers::ImportVertexColorsToSkeletalMesh(USkeletalMesh* Skel
 }
 
 
-bool UMeshPaintModeHelpers::RetrieveViewportPaintRays(const FSceneView* View, FViewport* Viewport, FPrimitiveDrawInterface* PDI, TArray<FPaintRay>& OutPaintRays)
+bool UMeshPaintModeSubsystem::RetrieveViewportPaintRays(const FSceneView* View, FViewport* Viewport, FPrimitiveDrawInterface* PDI, TArray<FPaintRay>& OutPaintRays)
 {
 	checkf(View && Viewport && PDI, TEXT("Invalid Viewport data"));
 	FEditorViewportClient* ViewportClient = (FEditorViewportClient*)Viewport->GetClient();
@@ -355,7 +355,7 @@ bool UMeshPaintModeHelpers::RetrieveViewportPaintRays(const FSceneView* View, FV
 }
 
 
-void UMeshPaintModeHelpers::ImportVertexColorsToStaticMesh(UStaticMesh* StaticMesh, const UImportVertexColorOptions* Options, UTexture2D* Texture)
+void UMeshPaintModeSubsystem::ImportVertexColorsToStaticMesh(UStaticMesh* StaticMesh, const UImportVertexColorOptions* Options, UTexture2D* Texture)
 {
 	checkf(StaticMesh && Options && Texture, TEXT("Invalid ptr"));
 
@@ -392,7 +392,7 @@ void UMeshPaintModeHelpers::ImportVertexColorsToStaticMesh(UStaticMesh* StaticMe
 	for (uint32 VertexIndex = 0; VertexIndex < LODModel.VertexBuffers.StaticMeshVertexBuffer.GetNumVertices(); ++VertexIndex)
 	{
 		const FVector2D UV = LODModel.VertexBuffers.StaticMeshVertexBuffer.GetVertexUV(VertexIndex, UVIndex);
-		LODModel.VertexBuffers.ColorVertexBuffer.VertexColor(VertexIndex) = UMeshPaintingToolset::PickVertexColorFromTextureData(MipData, UV, Texture, ColorMask);
+		LODModel.VertexBuffers.ColorVertexBuffer.VertexColor(VertexIndex) = GEngine->GetEngineSubsystem<UMeshPaintingSubsystem>()->PickVertexColorFromTextureData(MipData, UV, Texture, ColorMask);
 	}
 
 	// Make sure colors are saved into raw mesh
@@ -400,7 +400,7 @@ void UMeshPaintModeHelpers::ImportVertexColorsToStaticMesh(UStaticMesh* StaticMe
 	StaticMesh->InitResources();
 }
 
-void UMeshPaintModeHelpers::ImportVertexColorsToStaticMeshComponent(UStaticMeshComponent* StaticMeshComponent, const UImportVertexColorOptions* Options, UTexture2D* Texture)
+void UMeshPaintModeSubsystem::ImportVertexColorsToStaticMeshComponent(UStaticMeshComponent* StaticMeshComponent, const UImportVertexColorOptions* Options, UTexture2D* Texture)
 {
 	checkf(StaticMeshComponent && Options && Texture, TEXT("Invalid ptr"));
 
@@ -455,7 +455,7 @@ void UMeshPaintModeHelpers::ImportVertexColorsToStaticMeshComponent(UStaticMeshC
 		for (uint32 VertexIndex = 0; VertexIndex < LODModel.VertexBuffers.StaticMeshVertexBuffer.GetNumVertices(); ++VertexIndex)
 		{
 			const FVector2D UV = LODModel.VertexBuffers.StaticMeshVertexBuffer.GetVertexUV(VertexIndex, UVIndex);
-			InstanceMeshLODInfo.OverrideVertexColors->VertexColor(VertexIndex) = UMeshPaintingToolset::PickVertexColorFromTextureData(MipData, UV, Texture, ColorMask);
+			InstanceMeshLODInfo.OverrideVertexColors->VertexColor(VertexIndex) = GEngine->GetEngineSubsystem<UMeshPaintingSubsystem>()->PickVertexColorFromTextureData(MipData, UV, Texture, ColorMask);
 		}
 
 		//Update the cache painted vertices
@@ -470,7 +470,7 @@ void UMeshPaintModeHelpers::ImportVertexColorsToStaticMeshComponent(UStaticMeshC
 	}
 }
 
-void UMeshPaintModeHelpers::PropagateVertexColors(const TArray<UStaticMeshComponent *> StaticMeshComponents)
+void UMeshPaintModeSubsystem::PropagateVertexColors(const TArray<UStaticMeshComponent *> StaticMeshComponents)
 {
 	bool SomePaintWasPropagated = false;
 	TUniquePtr< FComponentReregisterContext > ComponentReregisterContext;
@@ -488,7 +488,7 @@ void UMeshPaintModeHelpers::PropagateVertexColors(const TArray<UStaticMeshCompon
 				{
 					Mesh->Modify();
 					// Try using the mapping generated when building the mesh.
-					if (UMeshPaintingToolset::PropagateColorsToRawMesh(Mesh, LODIndex, InstanceMeshLODInfo))
+					if (GEngine->GetEngineSubsystem<UMeshPaintingSubsystem>()->PropagateColorsToRawMesh(Mesh, LODIndex, InstanceMeshLODInfo))
 					{
 						SomePaintWasPropagated = true;
 					}
@@ -499,13 +499,13 @@ void UMeshPaintModeHelpers::PropagateVertexColors(const TArray<UStaticMeshCompon
 		if (SomePaintWasPropagated)
 		{
 			ComponentReregisterContext = MakeUnique<FComponentReregisterContext>(Component);
-			UMeshPaintingToolset::RemoveComponentInstanceVertexColors(Component);
+			GEngine->GetEngineSubsystem<UMeshPaintingSubsystem>()->RemoveComponentInstanceVertexColors(Component);
 			Mesh->Build();
 		}
 	}
 }
 
-bool UMeshPaintModeHelpers::CanPropagateVertexColors(TArray<UStaticMeshComponent*>& StaticMeshComponents, TArray<UStaticMesh*>& StaticMeshes, int32 NumInstanceVertexColorBytes)
+bool UMeshPaintModeSubsystem::CanPropagateVertexColors(TArray<UStaticMeshComponent*>& StaticMeshComponents, TArray<UStaticMesh*>& StaticMeshes, int32 NumInstanceVertexColorBytes)
 {
 	bool bValid = StaticMeshComponents.Num() > 0;
 	for (const UStaticMeshComponent* Component : StaticMeshComponents)
@@ -531,13 +531,13 @@ bool UMeshPaintModeHelpers::CanPropagateVertexColors(TArray<UStaticMeshComponent
 			CachedLODIndex = ColorPaintingTool->GetCachedLODIndex();
 		}
 
-		UMeshPaintingToolset::GetInstanceColorDataInfo(Component, CachedLODIndex, NumInstanceVertexColorBytes);
+		GEngine->GetEngineSubsystem<UMeshPaintingSubsystem>()->GetInstanceColorDataInfo(Component, CachedLODIndex, NumInstanceVertexColorBytes);
 	}
 
 	return bValid && (NumInstanceVertexColorBytes > 0);
 }
 
-void UMeshPaintModeHelpers::CopyVertexColors(const TArray<UStaticMeshComponent *> StaticMeshComponents, TArray<FPerComponentVertexColorData>& CopiedVertexColors)
+void UMeshPaintModeSubsystem::CopyVertexColors(const TArray<UStaticMeshComponent *> StaticMeshComponents, TArray<FPerComponentVertexColorData>& CopiedVertexColors)
 {
 	for (UStaticMeshComponent* Component : StaticMeshComponents)
 	{
@@ -562,13 +562,13 @@ void UMeshPaintModeHelpers::CopyVertexColors(const TArray<UStaticMeshComponent *
 
 				if (Component->LODData.IsValidIndex(LODIndex) && (Component->LODData[LODIndex].OverrideVertexColors != nullptr))
 				{
-					ColorData = UMeshPaintingToolset::GetInstanceColorDataForLOD(Component, LODIndex);
+					ColorData = GEngine->GetEngineSubsystem<UMeshPaintingSubsystem>()->GetInstanceColorDataForLOD(Component, LODIndex);
 				}
 				else
 				{
-					ColorData = UMeshPaintingToolset::GetColorDataForLOD(StaticMesh, LODIndex);
+					ColorData = GEngine->GetEngineSubsystem<UMeshPaintingSubsystem>()->GetColorDataForLOD(StaticMesh, LODIndex);
 				}
-				VertexData = UMeshPaintingToolset::GetVerticesForLOD(StaticMesh, LODIndex);
+				VertexData = GEngine->GetEngineSubsystem<UMeshPaintingSubsystem>()->GetVerticesForLOD(StaticMesh, LODIndex);
 
 				const bool bValidColorData = VertexData.Num() == ColorData.Num();
 				for (int32 VertexIndex = 0; VertexIndex < VertexData.Num(); ++VertexIndex)
@@ -584,7 +584,7 @@ void UMeshPaintModeHelpers::CopyVertexColors(const TArray<UStaticMeshComponent *
 	}
 }
 
-bool UMeshPaintModeHelpers::CanCopyInstanceVertexColors(const TArray<UStaticMeshComponent*>& StaticMeshComponents, int32 PaintingMeshLODIndex)
+bool UMeshPaintModeSubsystem::CanCopyInstanceVertexColors(const TArray<UStaticMeshComponent*>& StaticMeshComponents, int32 PaintingMeshLODIndex)
 {
 	// Ensure that the selection does not contain two components which point to identical meshes
 	TArray<const UStaticMesh*> ContainedMeshes;
@@ -615,7 +615,7 @@ bool UMeshPaintModeHelpers::CanCopyInstanceVertexColors(const TArray<UStaticMesh
 		checkf(Component != nullptr, TEXT("Invalid Static Mesh Component"));
 		if (Component->GetStaticMesh() != nullptr && Component->GetStaticMesh()->GetNumLODs() > (int32)PaintingMeshLODIndex)
 		{
-			uint32 BufferSize = UMeshPaintingToolset::GetVertexColorBufferSize(Component, PaintingMeshLODIndex, true);
+			uint32 BufferSize = GEngine->GetEngineSubsystem<UMeshPaintingSubsystem>()->GetVertexColorBufferSize(Component, PaintingMeshLODIndex, true);
 
 			if (BufferSize > 0)
 			{
@@ -627,7 +627,7 @@ bool UMeshPaintModeHelpers::CanCopyInstanceVertexColors(const TArray<UStaticMesh
 	return bValidSelection && (NumValidMeshes != 0);
 }
 
-void UMeshPaintModeHelpers::PasteVertexColors(const TArray<UStaticMeshComponent*>& StaticMeshComponents, TArray<FPerComponentVertexColorData>& CopiedColorsByComponent)
+void UMeshPaintModeSubsystem::PasteVertexColors(const TArray<UStaticMeshComponent*>& StaticMeshComponents, TArray<FPerComponentVertexColorData>& CopiedColorsByComponent)
 {
 	for (UStaticMeshComponent* Component : StaticMeshComponents)
 	{
@@ -652,7 +652,7 @@ void UMeshPaintModeHelpers::PasteVertexColors(const TArray<UStaticMeshComponent*
 				Component->Modify();
 				Component->SetLODDataCount(NumLods, NumLods);
 				/** Remove all vertex colors before we paste in new ones */
-				UMeshPaintingToolset::RemoveComponentInstanceVertexColors(Component);
+				GEngine->GetEngineSubsystem<UMeshPaintingSubsystem>()->RemoveComponentInstanceVertexColors(Component);
 
 				/** Try and apply copied vertex colors for each LOD in the mesh */
 				for (int32 LODIndex = 0; LODIndex < NumLods; ++LODIndex)
@@ -665,7 +665,7 @@ void UMeshPaintModeHelpers::PasteVertexColors(const TArray<UStaticMeshComponent*
 					{
 						// no corresponding LOD in color paste buffer CopiedColorsByLOD
 						// create array of all white verts
-						UMeshPaintingToolset::SetInstanceColorDataForLOD(Component, LODIndex, FColor::White, FColor::White);
+						GEngine->GetEngineSubsystem<UMeshPaintingSubsystem>()->SetInstanceColorDataForLOD(Component, LODIndex, FColor::White, FColor::White);
 					}
 					else
 					{
@@ -674,7 +674,7 @@ void UMeshPaintModeHelpers::PasteVertexColors(const TArray<UStaticMeshComponent*
 
 						if (NumLODVertices == LODData.ColorsByIndex.Num())
 						{
-							UMeshPaintingToolset::SetInstanceColorDataForLOD(Component, LODIndex, LODData.ColorsByIndex);
+							GEngine->GetEngineSubsystem<UMeshPaintingSubsystem>()->SetInstanceColorDataForLOD(Component, LODIndex, LODData.ColorsByIndex);
 						}
 						else
 						{
@@ -690,7 +690,7 @@ void UMeshPaintModeHelpers::PasteVertexColors(const TArray<UStaticMeshComponent*
 								PositionMatchedColors.Add(FoundColor ? *FoundColor : FColor::White);
 							}
 
-							UMeshPaintingToolset::SetInstanceColorDataForLOD(Component, LODIndex, PositionMatchedColors);
+							GEngine->GetEngineSubsystem<UMeshPaintingSubsystem>()->SetInstanceColorDataForLOD(Component, LODIndex, PositionMatchedColors);
 						}
 					}
 				}
@@ -703,7 +703,7 @@ void UMeshPaintModeHelpers::PasteVertexColors(const TArray<UStaticMeshComponent*
 	}
 }
 
-bool UMeshPaintModeHelpers::CanPasteInstanceVertexColors(const TArray<UStaticMeshComponent*>& StaticMeshComponents, const TArray<FPerComponentVertexColorData>& CopiedColorsByComponent)
+bool UMeshPaintModeSubsystem::CanPasteInstanceVertexColors(const TArray<UStaticMeshComponent*>& StaticMeshComponents, const TArray<FPerComponentVertexColorData>& CopiedColorsByComponent)
 {
 	bool bValidForPasting = false;
 	/** Make sure we have copied vertex color data which matches at least mesh component in the current selection */
@@ -731,7 +731,7 @@ bool UMeshPaintModeHelpers::CanPasteInstanceVertexColors(const TArray<UStaticMes
 	return bValidForPasting;
 }
 
-void UMeshPaintModeHelpers::RemovePerLODColors(const TArray<UMeshComponent*>& PaintableComponents)
+void UMeshPaintModeSubsystem::RemovePerLODColors(const TArray<UMeshComponent*>& PaintableComponents)
 {
 	//Remove painting on all lowers LODs before doing the propagation
 	for (UMeshComponent* SelectedComponent : PaintableComponents)
@@ -766,7 +766,7 @@ void UMeshPaintModeHelpers::RemovePerLODColors(const TArray<UMeshComponent*>& Pa
 	}
 }
 
-void UMeshPaintModeHelpers::SwapVertexColors()
+void UMeshPaintModeSubsystem::SwapVertexColors()
 {
 	UMeshVertexPaintingToolProperties* Settings = UMeshPaintMode::GetVertexToolProperties();
 	if (Settings)
@@ -781,7 +781,7 @@ void UMeshPaintModeHelpers::SwapVertexColors()
 
 
 
-void UMeshPaintModeHelpers::SaveModifiedTextures()
+void UMeshPaintModeSubsystem::SaveModifiedTextures()
 {
 	UMeshTexturePaintingToolProperties* Settings = UMeshPaintMode::GetTextureToolProperties();
 	if (Settings)
@@ -797,7 +797,7 @@ void UMeshPaintModeHelpers::SaveModifiedTextures()
 	}
 }
 
-bool UMeshPaintModeHelpers::CanSaveModifiedTextures()
+bool UMeshPaintModeSubsystem::CanSaveModifiedTextures()
 {
 	/** Check whether or not the current selected paint texture requires saving */
 	bool bRequiresSaving = false;
diff --git a/Engine/Plugins/MeshPainting/Source/MeshPaintEditorMode/Public/MeshPaintModeHelpers.h b/Engine/Plugins/MeshPainting/Source/MeshPaintEditorMode/Public/MeshPaintModeHelpers.h
--- a/Engine/Plugins/MeshPainting/Source/MeshPaintEditorMode/Public/MeshPaintModeHelpers.h
+++ b/Engine/Plugins/MeshPainting/Source/MeshPaintEditorMode/Public/MeshPaintModeHelpers.h
@@ -5,7 +5,7 @@
 #include "CoreMinimal.h"
 #include "Engine/EngineTypes.h"
 #include "MeshPaintingToolsetTypes.h"
-#include "Kismet/BlueprintFunctionLibrary.h"
+#include "EditorSubsystem.h"
 #include "IMeshPaintComponentAdapter.h"
 
 class FMeshPaintParameters;
@@ -29,21 +29,21 @@ struct FPerComponentVertexColorData;
 
 enum class EMeshPaintDataColorViewMode : uint8;
 
-class MESHPAINTEDITORMODE_API UMeshPaintModeHelpers : public UBlueprintFunctionLibrary
+class MESHPAINTEDITORMODE_API UMeshPaintModeSubsystem : public UEditorSubsystem
 {
 public:
 	/** Forces the Viewport Client to render using the given Viewport Color ViewMode */
-	static void SetViewportColorMode(EMeshPaintDataColorViewMode ColorViewMode, FEditorViewportClient* ViewportClient);
+	void SetViewportColorMode(EMeshPaintDataColorViewMode ColorViewMode, FEditorViewportClient* ViewportClient);
 
 	/** Sets whether or not the level viewport should be real time rendered move or viewport as parameter? */
-	static void SetRealtimeViewport(bool bRealtime);
+	void SetRealtimeViewport(bool bRealtime);
 
 
 	/** Helper function to import Vertex Colors from a Texture to the specified MeshComponent (makes use of SImportVertexColorsOptions Widget) */
-	static void ImportVertexColorsFromTexture(UMeshComponent* MeshComponent);
+	void ImportVertexColorsFromTexture(UMeshComponent* MeshComponent);
 
 	/** Imports vertex colors from a Texture to the specified Skeletal Mesh according to user-set options */
-	static void ImportVertexColorsToSkeletalMesh(USkeletalMesh* SkeletalMesh, const UImportVertexColorOptions* Options, UTexture2D* Texture);
+	void ImportVertexColorsToSkeletalMesh(USkeletalMesh* SkeletalMesh, const UImportVertexColorOptions* Options, UTexture2D* Texture);
 
 	struct FPaintRay
 	{
@@ -54,23 +54,23 @@ public:
 	};
 
 
-	static bool RetrieveViewportPaintRays(const FSceneView* View, FViewport* Viewport, FPrimitiveDrawInterface* PDI, TArray<FPaintRay>& OutPaintRays);
+	bool RetrieveViewportPaintRays(const FSceneView* View, FViewport* Viewport, FPrimitiveDrawInterface* PDI, TArray<FPaintRay>& OutPaintRays);
 
 	/** Imports vertex colors from a Texture to the specified Static Mesh according to user-set options */
-	static void ImportVertexColorsToStaticMesh(UStaticMesh* StaticMesh, const UImportVertexColorOptions* Options, UTexture2D* Texture);
+	void ImportVertexColorsToStaticMesh(UStaticMesh* StaticMesh, const UImportVertexColorOptions* Options, UTexture2D* Texture);
 
 	/** Imports vertex colors from a Texture to the specified Static Mesh Component according to user-set options */
-	static void ImportVertexColorsToStaticMeshComponent(UStaticMeshComponent* StaticMeshComponent, const UImportVertexColorOptions* Options, UTexture2D* Texture);
-
-	static void PropagateVertexColors(const TArray<UStaticMeshComponent *> StaticMeshComponents);
-	static bool CanPropagateVertexColors(TArray<UStaticMeshComponent*>& StaticMeshComponents, TArray<UStaticMesh*>& StaticMeshes, int32 NumInstanceVertexColorBytes);
-	static void CopyVertexColors(const TArray<UStaticMeshComponent*> StaticMeshComponents, TArray<FPerComponentVertexColorData>& CopiedVertexColors);
-	static bool CanCopyInstanceVertexColors(const TArray<UStaticMeshComponent*>& StaticMeshComponents, int32 PaintingMeshLODIndex);
-	static void PasteVertexColors(const TArray<UStaticMeshComponent*>& StaticMeshComponents, TArray<FPerComponentVertexColorData>& CopiedColorsByComponent);
-	static bool CanPasteInstanceVertexColors(const TArray<UStaticMeshComponent*>& StaticMeshComponents, const TArray<FPerComponentVertexColorData>& CopiedColorsByComponent);
-	static void RemovePerLODColors(const TArray<UMeshComponent*>& PaintableComponents);
-
-	static void SwapVertexColors();
-	static 	void SaveModifiedTextures();
-	static bool CanSaveModifiedTextures();
+	void ImportVertexColorsToStaticMeshComponent(UStaticMeshComponent* StaticMeshComponent, const UImportVertexColorOptions* Options, UTexture2D* Texture);
+
+	void PropagateVertexColors(const TArray<UStaticMeshComponent *> StaticMeshComponents);
+	bool CanPropagateVertexColors(TArray<UStaticMeshComponent*>& StaticMeshComponents, TArray<UStaticMesh*>& StaticMeshes, int32 NumInstanceVertexColorBytes);
+	void CopyVertexColors(const TArray<UStaticMeshComponent*> StaticMeshComponents, TArray<FPerComponentVertexColorData>& CopiedVertexColors);
+	bool CanCopyInstanceVertexColors(const TArray<UStaticMeshComponent*>& StaticMeshComponents, int32 PaintingMeshLODIndex);
+	void PasteVertexColors(const TArray<UStaticMeshComponent*>& StaticMeshComponents, TArray<FPerComponentVertexColorData>& CopiedColorsByComponent);
+	bool CanPasteInstanceVertexColors(const TArray<UStaticMeshComponent*>& StaticMeshComponents, const TArray<FPerComponentVertexColorData>& CopiedColorsByComponent);
+	void RemovePerLODColors(const TArray<UMeshComponent*>& PaintableComponents);
+
+	void SwapVertexColors();
+	void SaveModifiedTextures();
+	bool CanSaveModifiedTextures();
 };
\ No newline at end of file
diff --git a/Engine/Plugins/MeshPainting/Source/MeshPaintingToolset/MeshPaintingToolset.Build.cs b/Engine/Plugins/MeshPainting/Source/MeshPaintingToolset/MeshPaintingToolset.Build.cs
--- a/Engine/Plugins/MeshPainting/Source/MeshPaintingToolset/MeshPaintingToolset.Build.cs
+++ b/Engine/Plugins/MeshPainting/Source/MeshPaintingToolset/MeshPaintingToolset.Build.cs
@@ -29,7 +29,7 @@ namespace UnrealBuildTool.Rules
                 new string[] {
 					"InteractiveToolsFramework",
 					"EditorInteractiveToolsFramework",
-                    "GeometricObjects"
+                    "GeometricObjects",
                 }
             );
         }
diff --git a/Engine/Plugins/MeshPainting/Source/MeshPaintingToolset/Private/MeshPaintHelpers.cpp b/Engine/Plugins/MeshPainting/Source/MeshPaintingToolset/Private/MeshPaintHelpers.cpp
--- a/Engine/Plugins/MeshPainting/Source/MeshPaintingToolset/Private/MeshPaintHelpers.cpp
+++ b/Engine/Plugins/MeshPainting/Source/MeshPaintingToolset/Private/MeshPaintHelpers.cpp
@@ -27,12 +27,12 @@
 
 extern void PropagateVertexPaintToSkeletalMesh(USkeletalMesh* SkeletalMesh, int32 LODIndex);
 
-bool UMeshPaintingToolset::HasPaintableMesh(UActorComponent* Component)
+bool UMeshPaintingSubsystem::HasPaintableMesh(UActorComponent* Component)
 {
 	return Cast<UMeshComponent>(Component) != nullptr;
 }
 
-void UMeshPaintingToolset::RemoveInstanceVertexColors(UObject* Obj)
+void UMeshPaintingSubsystem::RemoveInstanceVertexColors(UObject* Obj)
 {
 	// Currently only static mesh component supports per instance vertex colors so only need to retrieve those and remove colors
 	AActor* Actor = Cast<AActor>(Obj);
@@ -44,13 +44,13 @@ void UMeshPaintingToolset::RemoveInstanceVertexColors(UObject* Obj)
 		{
 			if (StaticMeshComponent != nullptr)
 			{
-				UMeshPaintingToolset::RemoveComponentInstanceVertexColors(StaticMeshComponent);
+				UMeshPaintingSubsystem::RemoveComponentInstanceVertexColors(StaticMeshComponent);
 			}
 		}
 	}
 }
 
-void UMeshPaintingToolset::RemoveComponentInstanceVertexColors(UStaticMeshComponent* StaticMeshComponent)
+void UMeshPaintingSubsystem::RemoveComponentInstanceVertexColors(UStaticMeshComponent* StaticMeshComponent)
 {
 	if (StaticMeshComponent != nullptr && StaticMeshComponent->GetStaticMesh() != nullptr)
 	{
@@ -75,7 +75,7 @@ void UMeshPaintingToolset::RemoveComponentInstanceVertexColors(UStaticMeshCompon
 }
 
 
-bool UMeshPaintingToolset::PropagateColorsToRawMesh(UStaticMesh* StaticMesh, int32 LODIndex, FStaticMeshComponentLODInfo& ComponentLODInfo)
+bool UMeshPaintingSubsystem::PropagateColorsToRawMesh(UStaticMesh* StaticMesh, int32 LODIndex, FStaticMeshComponentLODInfo& ComponentLODInfo)
 {
 	check(ComponentLODInfo.OverrideVertexColors);
 	check(StaticMesh->IsSourceModelValid(LODIndex));
@@ -167,14 +167,14 @@ bool UMeshPaintingToolset::PropagateColorsToRawMesh(UStaticMesh* StaticMesh, int
 	return bPropagatedColors;
 }
 
-bool UMeshPaintingToolset::PaintVertex(const FVector& InVertexPosition, const FMeshPaintParameters& InParams, FColor& InOutVertexColor)
+bool UMeshPaintingSubsystem::PaintVertex(const FVector& InVertexPosition, const FMeshPaintParameters& InParams, FColor& InOutVertexColor)
 {
 	float SquaredDistanceToVertex2D;
 	float VertexDepthToBrush;
-	if (UMeshPaintingToolset::IsPointInfluencedByBrush(InVertexPosition, InParams, SquaredDistanceToVertex2D, VertexDepthToBrush))
+	if (UMeshPaintingSubsystem::IsPointInfluencedByBrush(InVertexPosition, InParams, SquaredDistanceToVertex2D, VertexDepthToBrush))
 	{
 		// Compute amount of paint to apply
-		const float PaintAmount = UMeshPaintingToolset::ComputePaintMultiplier(SquaredDistanceToVertex2D, InParams.BrushStrength, InParams.InnerBrushRadius, InParams.BrushRadialFalloffRange, InParams.BrushDepth, InParams.BrushDepthFalloffRange, VertexDepthToBrush);
+		const float PaintAmount = UMeshPaintingSubsystem::ComputePaintMultiplier(SquaredDistanceToVertex2D, InParams.BrushStrength, InParams.InnerBrushRadius, InParams.BrushRadialFalloffRange, InParams.BrushDepth, InParams.BrushDepthFalloffRange, VertexDepthToBrush);
 			
 		const FLinearColor OldColor = InOutVertexColor.ReinterpretAsLinear();
 		FLinearColor NewColor = OldColor;
@@ -194,7 +194,7 @@ bool UMeshPaintingToolset::PaintVertex(const FVector& InVertexPosition, const FM
 	return false;
 }
 
-void UMeshPaintingToolset::ApplyVertexColorPaint(const FMeshPaintParameters &InParams, const FLinearColor &OldColor, FLinearColor &NewColor, const float PaintAmount)
+void UMeshPaintingSubsystem::ApplyVertexColorPaint(const FMeshPaintParameters &InParams, const FLinearColor &OldColor, FLinearColor &NewColor, const float PaintAmount)
 {
 	// Color painting
 
@@ -219,7 +219,7 @@ void UMeshPaintingToolset::ApplyVertexColorPaint(const FMeshPaintParameters &InP
 	}
 }
 
-void UMeshPaintingToolset::ApplyVertexWeightPaint(const FMeshPaintParameters &InParams, const FLinearColor &OldColor, FLinearColor &NewColor, const float PaintAmount)
+void UMeshPaintingSubsystem::ApplyVertexWeightPaint(const FMeshPaintParameters &InParams, const FLinearColor &OldColor, FLinearColor &NewColor, const float PaintAmount)
 {
 	// Total number of texture blend weights we're using
 	check(InParams.TotalWeightCount > 0);
@@ -373,7 +373,7 @@ void UMeshPaintingToolset::ApplyVertexWeightPaint(const FMeshPaintParameters &In
 	}	
 }
 
-FLinearColor UMeshPaintingToolset::GenerateColorForTextureWeight(const int32 NumWeights, const int32 WeightIndex)
+FLinearColor UMeshPaintingSubsystem::GenerateColorForTextureWeight(const int32 NumWeights, const int32 WeightIndex)
 {
 	const bool bUsingOneMinusTotal =
 		NumWeights == 2 ||		// Two textures: Use a lerp() in pixel shader (single value)
@@ -470,7 +470,7 @@ FLinearColor UMeshPaintingToolset::GenerateColorForTextureWeight(const int32 Num
 	return NewColor;
 }
 
-float UMeshPaintingToolset::ComputePaintMultiplier(float SquaredDistanceToVertex2D, float BrushStrength, float BrushInnerRadius, float BrushRadialFalloff, float BrushInnerDepth, float BrushDepthFallof, float VertexDepthToBrush)
+float UMeshPaintingSubsystem::ComputePaintMultiplier(float SquaredDistanceToVertex2D, float BrushStrength, float BrushInnerRadius, float BrushRadialFalloff, float BrushInnerDepth, float BrushDepthFallof, float VertexDepthToBrush)
 {
 	float PaintAmount = 1.0f;
 	
@@ -500,7 +500,7 @@ float UMeshPaintingToolset::ComputePaintMultiplier(float SquaredDistanceToVertex
 	return PaintAmount;
 }
 
-bool UMeshPaintingToolset::IsPointInfluencedByBrush(const FVector& InPosition, const FMeshPaintParameters& InParams, float& OutSquaredDistanceToVertex2D, float& OutVertexDepthToBrush)
+bool UMeshPaintingSubsystem::IsPointInfluencedByBrush(const FVector& InPosition, const FMeshPaintParameters& InParams, float& OutSquaredDistanceToVertex2D, float& OutVertexDepthToBrush)
 {
 	// Project the vertex into the plane of the brush
 	FVector BrushSpaceVertexPosition = InParams.InverseBrushToWorldMatrix.TransformPosition(InPosition);
@@ -524,7 +524,7 @@ bool UMeshPaintingToolset::IsPointInfluencedByBrush(const FVector& InPosition, c
 	return false;
 }
 
-bool UMeshPaintingToolset::IsPointInfluencedByBrush(const FVector2D& BrushSpacePosition, const float BrushRadius, float& OutInRangeValue)
+bool UMeshPaintingSubsystem::IsPointInfluencedByBrush(const FVector2D& BrushSpacePosition, const float BrushRadius, float& OutInRangeValue)
 {
 	const float DistanceToBrush = BrushSpacePosition.SizeSquared();
 	if (DistanceToBrush <= BrushRadius)
@@ -536,7 +536,7 @@ bool UMeshPaintingToolset::IsPointInfluencedByBrush(const FVector2D& BrushSpaceP
 	return false;
 }
 
-uint32 UMeshPaintingToolset::GetVertexColorBufferSize(UMeshComponent* MeshComponent, int32 LODIndex, bool bInstance)
+uint32 UMeshPaintingSubsystem::GetVertexColorBufferSize(UMeshComponent* MeshComponent, int32 LODIndex, bool bInstance)
 {
 	checkf(MeshComponent != nullptr, TEXT("Invalid static mesh component ptr"));
 	
@@ -595,7 +595,7 @@ uint32 UMeshPaintingToolset::GetVertexColorBufferSize(UMeshComponent* MeshCompon
 	return SizeInBytes;
 }
 
-TArray<FVector> UMeshPaintingToolset::GetVerticesForLOD( const UStaticMesh* StaticMesh, int32 LODIndex)
+TArray<FVector> UMeshPaintingSubsystem::GetVerticesForLOD( const UStaticMesh* StaticMesh, int32 LODIndex)
 {
 	checkf(StaticMesh != nullptr, TEXT("Invalid static mesh ptr"));
 
@@ -614,7 +614,7 @@ TArray<FVector> UMeshPaintingToolset::GetVerticesForLOD( const UStaticMesh* Stat
 	return Vertices;
 }
 
-TArray<FColor> UMeshPaintingToolset::GetColorDataForLOD( const UStaticMesh* StaticMesh, int32 LODIndex)
+TArray<FColor> UMeshPaintingSubsystem::GetColorDataForLOD( const UStaticMesh* StaticMesh, int32 LODIndex)
 {
 	checkf(StaticMesh != nullptr, TEXT("Invalid static mesh ptr"));
 	// Retrieve mesh vertex colors from Static mesh render data 
@@ -632,7 +632,7 @@ TArray<FColor> UMeshPaintingToolset::GetColorDataForLOD( const UStaticMesh* Stat
 	return Colors;
 }
 
-TArray<FColor> UMeshPaintingToolset::GetInstanceColorDataForLOD(const UStaticMeshComponent* MeshComponent, int32 LODIndex)
+TArray<FColor> UMeshPaintingSubsystem::GetInstanceColorDataForLOD(const UStaticMeshComponent* MeshComponent, int32 LODIndex)
 {
 	checkf(MeshComponent != nullptr, TEXT("Invalid static mesh component ptr"));
 	TArray<FColor> Colors;
@@ -655,7 +655,7 @@ TArray<FColor> UMeshPaintingToolset::GetInstanceColorDataForLOD(const UStaticMes
 	return Colors;
 }
 
-void UMeshPaintingToolset::SetInstanceColorDataForLOD(UStaticMeshComponent* MeshComponent, int32 LODIndex, const TArray<FColor>& Colors)
+void UMeshPaintingSubsystem::SetInstanceColorDataForLOD(UStaticMeshComponent* MeshComponent, int32 LODIndex, const TArray<FColor>& Colors)
 {
 	checkf(MeshComponent != nullptr, TEXT("Invalid static mesh component ptr"));
 
@@ -689,7 +689,7 @@ void UMeshPaintingToolset::SetInstanceColorDataForLOD(UStaticMeshComponent* Mesh
 	}
 }
 
-void UMeshPaintingToolset::SetInstanceColorDataForLOD(UStaticMeshComponent* MeshComponent, int32 LODIndex, const FColor FillColor, const FColor MaskColor )
+void UMeshPaintingSubsystem::SetInstanceColorDataForLOD(UStaticMeshComponent* MeshComponent, int32 LODIndex, const FColor FillColor, const FColor MaskColor )
 {
 	checkf(MeshComponent != nullptr, TEXT("Invalid static mesh component ptr"));
 
@@ -749,7 +749,7 @@ void UMeshPaintingToolset::SetInstanceColorDataForLOD(UStaticMeshComponent* Mesh
 	}
 }
 
-void UMeshPaintingToolset::FillStaticMeshVertexColors(UStaticMeshComponent* MeshComponent, int32 LODIndex, const FColor FillColor, const FColor MaskColor)
+void UMeshPaintingSubsystem::FillStaticMeshVertexColors(UStaticMeshComponent* MeshComponent, int32 LODIndex, const FColor FillColor, const FColor MaskColor)
 {
 	UStaticMesh* Mesh = MeshComponent->GetStaticMesh();
 	if (Mesh)
@@ -761,18 +761,18 @@ void UMeshPaintingToolset::FillStaticMeshVertexColors(UStaticMeshComponent* Mesh
 			{
 				for (LODIndex = 0; LODIndex < NumLODs; ++LODIndex)
 				{
-					UMeshPaintingToolset::SetInstanceColorDataForLOD(MeshComponent, LODIndex, FillColor, MaskColor);
+					UMeshPaintingSubsystem::SetInstanceColorDataForLOD(MeshComponent, LODIndex, FillColor, MaskColor);
 				}
 			}
 			else
 			{
-				UMeshPaintingToolset::SetInstanceColorDataForLOD(MeshComponent, LODIndex, FillColor, MaskColor);
+				UMeshPaintingSubsystem::SetInstanceColorDataForLOD(MeshComponent, LODIndex, FillColor, MaskColor);
 			}
 		}
 	}
 }
 
-void UMeshPaintingToolset::FillSkeletalMeshVertexColors(USkeletalMeshComponent* MeshComponent, int32 LODIndex, const FColor FillColor, const FColor MaskColor)
+void UMeshPaintingSubsystem::FillSkeletalMeshVertexColors(USkeletalMeshComponent* MeshComponent, int32 LODIndex, const FColor FillColor, const FColor MaskColor)
 {
 	TUniquePtr< FSkinnedMeshComponentRecreateRenderStateContext > RecreateRenderStateContext;
 	USkeletalMesh* Mesh = MeshComponent->SkeletalMesh;
@@ -799,7 +799,7 @@ void UMeshPaintingToolset::FillSkeletalMeshVertexColors(USkeletalMeshComponent*
 			// because currently all LOD data is being released above.
 			for (LODIndex = 0; LODIndex < NumLODs; ++LODIndex)
 			{
-				UMeshPaintingToolset::SetColorDataForLOD(Mesh, LODIndex, FillColor, MaskColor);
+				UMeshPaintingSubsystem::SetColorDataForLOD(Mesh, LODIndex, FillColor, MaskColor);
 				PropagateVertexPaintToSkeletalMesh(Mesh, LODIndex);
 			}
 			Mesh->InitResources();
@@ -807,7 +807,7 @@ void UMeshPaintingToolset::FillSkeletalMeshVertexColors(USkeletalMeshComponent*
 	}
 }
 
-void UMeshPaintingToolset::SetColorDataForLOD(USkeletalMesh* SkeletalMesh, int32 LODIndex, const FColor FillColor, const FColor MaskColor )
+void UMeshPaintingSubsystem::SetColorDataForLOD(USkeletalMesh* SkeletalMesh, int32 LODIndex, const FColor FillColor, const FColor MaskColor )
 {
 	checkf(SkeletalMesh != nullptr, TEXT("Invalid Skeletal Mesh Ptr"));
 	FSkeletalMeshRenderData* Resource = SkeletalMesh->GetResourceForRendering();
@@ -850,7 +850,7 @@ void UMeshPaintingToolset::SetColorDataForLOD(USkeletalMesh* SkeletalMesh, int32
 	}
 }
 
-void UMeshPaintingToolset::ApplyFillWithMask(FColor& InOutColor, const FColor& MaskColor, const FColor& FillColor)
+void UMeshPaintingSubsystem::ApplyFillWithMask(FColor& InOutColor, const FColor& MaskColor, const FColor& FillColor)
 {
 	InOutColor.R = ((InOutColor.R & (~MaskColor.R)) | (FillColor.R & MaskColor.R));
 	InOutColor.G = ((InOutColor.G & (~MaskColor.G)) | (FillColor.G & MaskColor.G));
@@ -859,7 +859,7 @@ void UMeshPaintingToolset::ApplyFillWithMask(FColor& InOutColor, const FColor& M
 }
 
 
-void UMeshPaintingToolset::ForceRenderMeshLOD(UMeshComponent* Component, int32 LODIndex)
+void UMeshPaintingSubsystem::ForceRenderMeshLOD(UMeshComponent* Component, int32 LODIndex)
 {
 	if (UStaticMeshComponent* StaticMeshComponent = Cast<UStaticMeshComponent>(Component))
 	{
@@ -871,7 +871,7 @@ void UMeshPaintingToolset::ForceRenderMeshLOD(UMeshComponent* Component, int32 L
 	}	
 }
 
-void UMeshPaintingToolset::ClearMeshTextureOverrides(const IMeshPaintComponentAdapter& GeometryInfo, UMeshComponent* InMeshComponent)
+void UMeshPaintingSubsystem::ClearMeshTextureOverrides(const IMeshPaintComponentAdapter& GeometryInfo, UMeshComponent* InMeshComponent)
 {
 	if (InMeshComponent != nullptr)
 	{
@@ -888,7 +888,7 @@ void UMeshPaintingToolset::ClearMeshTextureOverrides(const IMeshPaintComponentAd
 	}
 }
 
-void UMeshPaintingToolset::ApplyVertexColorsToAllLODs(IMeshPaintComponentAdapter& GeometryInfo, UMeshComponent* InMeshComponent)
+void UMeshPaintingSubsystem::ApplyVertexColorsToAllLODs(IMeshPaintComponentAdapter& GeometryInfo, UMeshComponent* InMeshComponent)
 {
 	if (UStaticMeshComponent* StaticMeshComponent = Cast<UStaticMeshComponent>(InMeshComponent))
 	{
@@ -900,7 +900,7 @@ void UMeshPaintingToolset::ApplyVertexColorsToAllLODs(IMeshPaintComponentAdapter
 	}
 }
 
-void UMeshPaintingToolset::ApplyVertexColorsToAllLODs(IMeshPaintComponentAdapter& GeometryInfo, UStaticMeshComponent* StaticMeshComponent)
+void UMeshPaintingSubsystem::ApplyVertexColorsToAllLODs(IMeshPaintComponentAdapter& GeometryInfo, UStaticMeshComponent* StaticMeshComponent)
 {
 	// If a static mesh component was found, apply LOD0 painting to all lower LODs.
 	if (!StaticMeshComponent || !StaticMeshComponent->GetStaticMesh())
@@ -983,7 +983,7 @@ void UMeshPaintingToolset::ApplyVertexColorsToAllLODs(IMeshPaintComponentAdapter
 	}
 }
 
-bool UMeshPaintingToolset::TryGetNumberOfLODs(const UMeshComponent* MeshComponent, int32& OutNumLODs)
+bool UMeshPaintingSubsystem::TryGetNumberOfLODs(const UMeshComponent* MeshComponent, int32& OutNumLODs)
 {
 	if (const UStaticMeshComponent* StaticMeshComponent = Cast<UStaticMeshComponent>(MeshComponent))
 	{
@@ -1007,14 +1007,14 @@ bool UMeshPaintingToolset::TryGetNumberOfLODs(const UMeshComponent* MeshComponen
 	return false;
 }
 
-int32 UMeshPaintingToolset::GetNumberOfLODs(const UMeshComponent* MeshComponent)
+int32 UMeshPaintingSubsystem::GetNumberOfLODs(const UMeshComponent* MeshComponent)
 {
 	int32 NumLODs = 1;
 	TryGetNumberOfLODs(MeshComponent, NumLODs);
 	return NumLODs;
 }
 
-int32 UMeshPaintingToolset::GetNumberOfUVs(const UMeshComponent* MeshComponent, int32 LODIndex)
+int32 UMeshPaintingSubsystem::GetNumberOfUVs(const UMeshComponent* MeshComponent, int32 LODIndex)
 {
 	int32 NumUVs = 0;
 
@@ -1038,7 +1038,7 @@ int32 UMeshPaintingToolset::GetNumberOfUVs(const UMeshComponent* MeshComponent,
 	return NumUVs;
 }
 
-bool UMeshPaintingToolset::DoesMeshComponentContainPerLODColors(const UMeshComponent* MeshComponent)
+bool UMeshPaintingSubsystem::DoesMeshComponentContainPerLODColors(const UMeshComponent* MeshComponent)
 {
 	bool bPerLODColors = false;
 
@@ -1087,7 +1087,7 @@ bool UMeshPaintingToolset::DoesMeshComponentContainPerLODColors(const UMeshCompo
 	return bPerLODColors;
 }
 
-void UMeshPaintingToolset::GetInstanceColorDataInfo(const UStaticMeshComponent* StaticMeshComponent, int32 LODIndex, int32& OutTotalInstanceVertexColorBytes)
+void UMeshPaintingSubsystem::GetInstanceColorDataInfo(const UStaticMeshComponent* StaticMeshComponent, int32 LODIndex, int32& OutTotalInstanceVertexColorBytes)
 {
 	checkf(StaticMeshComponent, TEXT("Invalid StaticMeshComponent"));
 	OutTotalInstanceVertexColorBytes = 0;
@@ -1104,7 +1104,7 @@ void UMeshPaintingToolset::GetInstanceColorDataInfo(const UStaticMeshComponent*
 	}
 }
 
-FColor UMeshPaintingToolset::PickVertexColorFromTextureData(const uint8* MipData, const FVector2D& UVCoordinate, const UTexture2D* Texture, const FColor ColorMask)
+FColor UMeshPaintingSubsystem::PickVertexColorFromTextureData(const uint8* MipData, const FVector2D& UVCoordinate, const UTexture2D* Texture, const FColor ColorMask)
 {
 	checkf(MipData, TEXT("Invalid texture MIP data"));
 	FColor VertexColor = FColor::Black;
@@ -1126,7 +1126,7 @@ FColor UMeshPaintingToolset::PickVertexColorFromTextureData(const uint8* MipData
 	return VertexColor;
 }
 
-bool UMeshPaintingToolset::GetPerVertexPaintInfluencedVertices(FPerVertexPaintActionArgs& InArgs, TSet<int32>& InfluencedVertices)
+bool UMeshPaintingSubsystem::GetPerVertexPaintInfluencedVertices(FPerVertexPaintActionArgs& InArgs, TSet<int32>& InfluencedVertices)
 {
 	// Retrieve components world matrix
 	const FMatrix& ComponentToWorldMatrix = InArgs.Adapter->GetComponentToWorldMatrix();
@@ -1146,7 +1146,7 @@ bool UMeshPaintingToolset::GetPerVertexPaintInfluencedVertices(FPerVertexPaintAc
 	return (InfluencedVertices.Num() > 0);
 }
 
-bool UMeshPaintingToolset::ApplyPerVertexPaintAction(FPerVertexPaintActionArgs& InArgs, FPerVertexPaintAction Action)
+bool UMeshPaintingSubsystem::ApplyPerVertexPaintAction(FPerVertexPaintActionArgs& InArgs, FPerVertexPaintAction Action)
 {
 	// Get a list of unique vertices indexed by the influenced triangles
 	TSet<int32> InfluencedVertices;
@@ -1166,7 +1166,7 @@ bool UMeshPaintingToolset::ApplyPerVertexPaintAction(FPerVertexPaintActionArgs&
 	return (InfluencedVertices.Num() > 0);
 }
 
-bool UMeshPaintingToolset::ApplyPerTrianglePaintAction(IMeshPaintComponentAdapter* Adapter, const FVector& CameraPosition, const FVector& HitPosition, const UBrushBaseProperties* Settings, FPerTrianglePaintAction Action, bool bOnlyFrontFacingTriangles)
+bool UMeshPaintingSubsystem::ApplyPerTrianglePaintAction(IMeshPaintComponentAdapter* Adapter, const FVector& CameraPosition, const FVector& HitPosition, const UBrushBaseProperties* Settings, FPerTrianglePaintAction Action, bool bOnlyFrontFacingTriangles)
 {
 	// Retrieve components world matrix
 	const FMatrix& ComponentToWorldMatrix = Adapter->GetComponentToWorldMatrix();
@@ -1253,7 +1253,7 @@ struct FVertexColorPropogationOctreeSemantics
 };
 typedef TOctree2<FPaintedMeshVertex, FVertexColorPropogationOctreeSemantics> TVertexColorPropogationPosOctree;
 
-void UMeshPaintingToolset::ApplyVertexColorsToAllLODs(IMeshPaintComponentAdapter& GeometryInfo, USkeletalMeshComponent* SkeletalMeshComponent)
+void UMeshPaintingSubsystem::ApplyVertexColorsToAllLODs(IMeshPaintComponentAdapter& GeometryInfo, USkeletalMeshComponent* SkeletalMeshComponent)
 {
 	checkf(SkeletalMeshComponent != nullptr, TEXT("Invalid Skeletal Mesh Component"));
 	USkeletalMesh* Mesh = SkeletalMeshComponent->SkeletalMesh;
@@ -1534,9 +1534,9 @@ void UMeshToolManager::CacheSelectionData(const int32 PaintLODIndex, const int32
 			AddPaintableMeshComponent(MeshComponent);
 			AddToComponentToAdapterMap(MeshComponent, MeshAdapter);
 			MeshAdapter->OnAdded();
-			UMeshPaintingToolset::ForceRenderMeshLOD(MeshComponent, PaintLODIndex);
+			GEngine->GetEngineSubsystem<UMeshPaintingSubsystem>()->ForceRenderMeshLOD(MeshComponent, PaintLODIndex);
 			ComponentReregisterContext.Reset(new FComponentReregisterContext(MeshComponent));
-			bSelectionContainsPerLODColors |= UMeshPaintingToolset::DoesMeshComponentContainPerLODColors(MeshComponent);
+			bSelectionContainsPerLODColors |= GEngine->GetEngineSubsystem<UMeshPaintingSubsystem>()->DoesMeshComponentContainPerLODColors(MeshComponent);
 		}
 	}
 	bNeedsRecache = false;
@@ -1546,7 +1546,7 @@ int32 UMeshToolManager::GetMaxUVIndexToPaint() const
 {
 	if (PaintableComponents.Num() == 1 && PaintableComponents[0])
 	{
-		return UMeshPaintingToolset::GetNumberOfUVs(PaintableComponents[0], 0) - 1;
+		return GEngine->GetEngineSubsystem<UMeshPaintingSubsystem>()->GetNumberOfUVs(PaintableComponents[0], 0) - 1;
 	}
 
 	return 0;
diff --git a/Engine/Plugins/MeshPainting/Source/MeshPaintingToolset/Private/MeshTexturePaintingTool.cpp b/Engine/Plugins/MeshPainting/Source/MeshPaintingToolset/Private/MeshTexturePaintingTool.cpp
--- a/Engine/Plugins/MeshPainting/Source/MeshPaintingToolset/Private/MeshTexturePaintingTool.cpp
+++ b/Engine/Plugins/MeshPainting/Source/MeshPaintingToolset/Private/MeshTexturePaintingTool.cpp
@@ -30,7 +30,11 @@
 
 bool UMeshTexturePaintingToolBuilder::CanBuildTool(const FToolBuilderState& SceneState) const
 {
-	UActorComponent* ActorComponent = ToolBuilderUtil::FindFirstComponent(SceneState, UMeshPaintingToolset::HasPaintableMesh);
+	const TFunction<bool(UActorComponent*)>& HasPaintableMesh = [this](UActorComponent* InComponent)
+	{ 
+		return GEngine->GetEngineSubsystem<UMeshPaintingSubsystem>()->HasPaintableMesh(InComponent);
+	};
+	UActorComponent* ActorComponent = ToolBuilderUtil::FindFirstComponent(SceneState, HasPaintableMesh);
 	TArray<FPaintableTexture> PaintableTextures;
 	if (ActorComponent && SharedMeshToolData.IsValid())
 	{
@@ -438,7 +442,7 @@ bool UMeshTexturePaintingTool::PaintInternal(const TArrayView<TPair<FVector, FVe
 						const FVector& BestTraceResultLocation = PaintRayResults[PaintRayResultId].BestTraceResult.Location;
 						FViewCameraState CameraState;
 						GetToolManager()->GetContextQueriesAPI()->GetCurrentViewState(CameraState);
-						bPaintApplied |= UMeshPaintingToolset::ApplyPerTrianglePaintAction(MeshAdapter, CameraState.Position, BestTraceResultLocation, BrushProperties, FPerTrianglePaintAction::CreateUObject(this, &UMeshTexturePaintingTool::GatherTextureTriangles, &TrianglePaintInfoArray, &MaterialSections, TextureProperties->UVChannel), TextureProperties->bOnlyFrontFacingTriangles);
+						bPaintApplied |= GEngine->GetEngineSubsystem<UMeshPaintingSubsystem>()->ApplyPerTrianglePaintAction(MeshAdapter, CameraState.Position, BestTraceResultLocation, BrushProperties, FPerTrianglePaintAction::CreateUObject(this, &UMeshTexturePaintingTool::GatherTextureTriangles, &TrianglePaintInfoArray, &MaterialSections, TextureProperties->UVChannel), TextureProperties->bOnlyFrontFacingTriangles);
 						break;
 					}
 
diff --git a/Engine/Plugins/MeshPainting/Source/MeshPaintingToolset/Private/MeshVertexPaintingTool.cpp b/Engine/Plugins/MeshPainting/Source/MeshPaintingToolset/Private/MeshVertexPaintingTool.cpp
--- a/Engine/Plugins/MeshPainting/Source/MeshPaintingToolset/Private/MeshVertexPaintingTool.cpp
+++ b/Engine/Plugins/MeshPainting/Source/MeshPaintingToolset/Private/MeshVertexPaintingTool.cpp
@@ -394,7 +394,7 @@ bool UMeshVertexPaintingTool::PaintInternal(const TArrayView<TPair<FVector, FVec
 				{
 					InfluencedVertices.Reset();
 					Args.HitResult = PaintRayResults[PaintRayResultId].BestTraceResult;
-					bPaintApplied |= UMeshPaintingToolset::GetPerVertexPaintInfluencedVertices(Args, InfluencedVertices);
+					bPaintApplied |= GEngine->GetEngineSubsystem<UMeshPaintingSubsystem>()->GetPerVertexPaintInfluencedVertices(Args, InfluencedVertices);
 
 					if (InfluencedVertices.Num() == 0)
 					{
@@ -432,7 +432,7 @@ void UMeshVertexPaintingTool::ApplyVertexData(FPerVertexPaintActionArgs& InArgs,
 	InArgs.Adapter->GetVertexPosition(VertexIndex, Position);
 	Position = InArgs.Adapter->GetComponentToWorldMatrix().TransformPosition(Position);
 	InArgs.Adapter->GetVertexColor(VertexIndex, PaintColor, true);
-	UMeshPaintingToolset::PaintVertex(Position, Parameters, PaintColor);
+	GEngine->GetEngineSubsystem<UMeshPaintingSubsystem>()->PaintVertex(Position, Parameters, PaintColor);
 	InArgs.Adapter->SetVertexColor(VertexIndex, PaintColor, true);
 }
 
@@ -599,7 +599,7 @@ void UMeshColorPaintingTool::SetAdditionalPaintParameters(FMeshPaintParameters&
 	InPaintParameters.bWriteGreen = ColorProperties->bWriteGreen;
 	InPaintParameters.bWriteBlue = ColorProperties->bWriteBlue;
 	InPaintParameters.bWriteAlpha = ColorProperties->bWriteAlpha;
-	InPaintParameters.ApplyVertexDataDelegate.AddStatic(&UMeshPaintingToolset::ApplyVertexColorPaint);
+	InPaintParameters.ApplyVertexDataDelegate.AddUObject(GEngine->GetEngineSubsystem<UMeshPaintingSubsystem>(), &UMeshPaintingSubsystem::ApplyVertexColorPaint);
 }
 
 int32 UMeshColorPaintingTool::GetMaxLODIndexToPaint() const
@@ -614,7 +614,7 @@ int32 UMeshColorPaintingTool::GetMaxLODIndexToPaint() const
 		for (UMeshComponent* MeshComponent : SelectedComponents)
 		{
 			int32 NumMeshLODs = 0;
-			if (UMeshPaintingToolset::TryGetNumberOfLODs(MeshComponent, NumMeshLODs))
+			if (GEngine->GetEngineSubsystem<UMeshPaintingSubsystem>()->TryGetNumberOfLODs(MeshComponent, NumMeshLODs))
 			{
 				ensure(NumMeshLODs > 0);
 				LODMin = FMath::Min(LODMin, NumMeshLODs - 1);
@@ -673,7 +673,7 @@ void UMeshColorPaintingTool::ApplyForcedLODIndex(int32 ForcedLODIndex)
 		{
 			if (SelectedComponent)
 			{
-				UMeshPaintingToolset::ForceRenderMeshLOD(SelectedComponent, ForcedLODIndex);
+				GEngine->GetEngineSubsystem<UMeshPaintingSubsystem>()->ForceRenderMeshLOD(SelectedComponent, ForcedLODIndex);
 			}
 		}
 	}
@@ -768,7 +768,7 @@ void UMeshWeightPaintingTool::SetAdditionalPaintParameters(FMeshPaintParameters&
 		// Clamp the weight index to fall within the total weight count
 		InPaintParameters.PaintWeightIndex = FMath::Clamp(PaintWeightIndex, 0, InPaintParameters.TotalWeightCount - 1);
 	}
-	InPaintParameters.ApplyVertexDataDelegate.AddStatic(&UMeshPaintingToolset::ApplyVertexWeightPaint);
+	InPaintParameters.ApplyVertexDataDelegate.AddUObject(GEngine->GetEngineSubsystem<UMeshPaintingSubsystem>(), &UMeshPaintingSubsystem::ApplyVertexWeightPaint);
 }
 
 #undef LOCTEXT_NAMESPACE
diff --git a/Engine/Plugins/MeshPainting/Source/MeshPaintingToolset/Public/MeshPaintHelpers.h b/Engine/Plugins/MeshPainting/Source/MeshPaintingToolset/Public/MeshPaintHelpers.h
--- a/Engine/Plugins/MeshPainting/Source/MeshPaintingToolset/Public/MeshPaintHelpers.h
+++ b/Engine/Plugins/MeshPainting/Source/MeshPaintingToolset/Public/MeshPaintHelpers.h
@@ -10,6 +10,7 @@
 #include "InteractiveToolManager.h"
 #include "EdModeInteractiveToolsContext.h"
 #include "Engine/StaticMesh.h"
+#include "Subsystems/EngineSubsystem.h"
 #include "MeshPaintHelpers.generated.h"
 
 class FMeshPaintParameters;
@@ -122,119 +123,115 @@ DECLARE_DELEGATE_TwoParams(FPerVertexPaintAction, FPerVertexPaintActionArgs& /*A
 DECLARE_DELEGATE_ThreeParams(FPerTrianglePaintAction, IMeshPaintComponentAdapter* /*Adapter*/, int32 /*TriangleIndex*/, const int32[3] /*Vertex Indices*/);
 
 UCLASS()
-class MESHPAINTINGTOOLSET_API UMeshPaintingToolset : public UBlueprintFunctionLibrary
+class MESHPAINTINGTOOLSET_API UMeshPaintingSubsystem : public UEngineSubsystem
 {
 
 	GENERATED_BODY()
 public:
-	static bool HasPaintableMesh(UActorComponent* Component);
+	bool HasPaintableMesh(UActorComponent* Component);
 	/** Removes vertex colors associated with the object */
-	static void RemoveInstanceVertexColors(UObject* Obj);
+	void RemoveInstanceVertexColors(UObject* Obj);
 
-	/** Removes vertex colors associated with the static mesh component */
-	static void RemoveComponentInstanceVertexColors(UStaticMeshComponent* StaticMeshComponent);
+	/** Removes vertex colors associated with the mesh component */
+	void RemoveComponentInstanceVertexColors(UStaticMeshComponent* StaticMeshComponent);
 	
-	/** Propagates per-instance vertex colors to the underlying Static Mesh for the given LOD Index */
-	static bool PropagateColorsToRawMesh(UStaticMesh* StaticMesh, int32 LODIndex, FStaticMeshComponentLODInfo& ComponentLODInfo);	
+	/** Propagates per-instance vertex colors to the underlying Mesh for the given LOD Index */
+	bool PropagateColorsToRawMesh(UStaticMesh* StaticMesh, int32 LODIndex, FStaticMeshComponentLODInfo& ComponentLODInfo);	
 
-	/** Retrieves the Vertex Color buffer size for the given LOD level in the Static Mesh */
-	static uint32 GetVertexColorBufferSize(UMeshComponent* MeshComponent, int32 LODIndex, bool bInstance);
+	/** Retrieves the Vertex Color buffer size for the given LOD level in the Mesh */
+	uint32 GetVertexColorBufferSize(UMeshComponent* MeshComponent, int32 LODIndex, bool bInstance);
 
-	/** Retrieves the vertex positions from the given LOD level in the Static Mesh */
-	static TArray<FVector> GetVerticesForLOD(const UStaticMesh* StaticMesh, int32 LODIndex);
+	/** Retrieves the vertex positions from the given LOD level in the Mesh */
+	TArray<FVector> GetVerticesForLOD(const UStaticMesh* StaticMesh, int32 LODIndex);
 
-	/** Retrieves the vertex colors from the given LOD level in the Static Mesh */
-	static TArray<FColor> GetColorDataForLOD(const UStaticMesh* StaticMesh, int32 LODIndex);
+	/** Retrieves the vertex colors from the given LOD level in the Mesh */
+	TArray<FColor> GetColorDataForLOD(const UStaticMesh* StaticMesh, int32 LODIndex);
 
 	/** Retrieves the per-instance vertex colors from the given LOD level in the StaticMeshComponent */
-	static TArray<FColor> GetInstanceColorDataForLOD(const UStaticMeshComponent* MeshComponent, int32 LODIndex);
+	TArray<FColor> GetInstanceColorDataForLOD(const UStaticMeshComponent* MeshComponent, int32 LODIndex);
 
 	/** Sets the specific (LOD Index) per-instance vertex colors for the given StaticMeshComponent to the supplied Color array */
-	static void SetInstanceColorDataForLOD(UStaticMeshComponent* MeshComponent, int32 LODIndex, const TArray<FColor>& Colors);	
+	void SetInstanceColorDataForLOD(UStaticMeshComponent* MeshComponent, int32 LODIndex, const TArray<FColor>& Colors);	
 	
 	/** Sets the specific (LOD Index) per-instance vertex colors for the given StaticMeshComponent to a single Color value */
-	static void SetInstanceColorDataForLOD(UStaticMeshComponent* MeshComponent, int32 LODIndex, const FColor FillColor, const FColor MaskColor);
+	void SetInstanceColorDataForLOD(UStaticMeshComponent* MeshComponent, int32 LODIndex, const FColor FillColor, const FColor MaskColor);
 	
 	/** Fills all vertex colors for all LODs found in the given mesh component with Fill Color */
-	static void FillStaticMeshVertexColors(UStaticMeshComponent* MeshComponent, int32 LODIndex, const FColor FillColor, const FColor MaskColor);
-	static void FillSkeletalMeshVertexColors(USkeletalMeshComponent* MeshComponent, int32 LODIndex, const FColor FillColor, const FColor MaskColor);
+	void FillStaticMeshVertexColors(UStaticMeshComponent* MeshComponent, int32 LODIndex, const FColor FillColor, const FColor MaskColor);
+	void FillSkeletalMeshVertexColors(USkeletalMeshComponent* MeshComponent, int32 LODIndex, const FColor FillColor, const FColor MaskColor);
 	
 	/** Sets all vertex colors for a specific LOD level in the SkeletalMesh to FillColor */
-	static void SetColorDataForLOD(USkeletalMesh* SkeletalMesh, int32 LODIndex, const FColor FillColor, const FColor MaskColor);
+	void SetColorDataForLOD(USkeletalMesh* SkeletalMesh, int32 LODIndex, const FColor FillColor, const FColor MaskColor);
 
-	static void ApplyFillWithMask(FColor& InOutColor, const FColor& MaskColor, const FColor& FillColor);
+	void ApplyFillWithMask(FColor& InOutColor, const FColor& MaskColor, const FColor& FillColor);
 
 	/** Forces the component to render LOD level at LODIndex instead of the view-based LOD level ( X = 0 means do not force the LOD, X > 0 means force the lod to X - 1 ) */
-	static void ForceRenderMeshLOD(UMeshComponent* Component, int32 LODIndex);
+	void ForceRenderMeshLOD(UMeshComponent* Component, int32 LODIndex);
 
 	/** Clears all texture overrides for this component. */
-	static void ClearMeshTextureOverrides(const IMeshPaintComponentAdapter& GeometryInfo, UMeshComponent* InMeshComponent);
+	void ClearMeshTextureOverrides(const IMeshPaintComponentAdapter& GeometryInfo, UMeshComponent* InMeshComponent);
 
 	/** Applies vertex color painting found on LOD 0 to all lower LODs. */
-	static void ApplyVertexColorsToAllLODs(IMeshPaintComponentAdapter& GeometryInfo, UMeshComponent* InMeshComponent);
+	void ApplyVertexColorsToAllLODs(IMeshPaintComponentAdapter& GeometryInfo, UMeshComponent* InMeshComponent);
 
 	/** Applies the vertex colors found in LOD level 0 to all contained LOD levels in the StaticMeshComponent */
-	static void ApplyVertexColorsToAllLODs(IMeshPaintComponentAdapter& GeometryInfo, UStaticMeshComponent* StaticMeshComponent);
+	void ApplyVertexColorsToAllLODs(IMeshPaintComponentAdapter& GeometryInfo, UStaticMeshComponent* StaticMeshComponent);
 
 	/** Applies the vertex colors found in LOD level 0 to all contained LOD levels in the SkeletalMeshComponent */
-	static void ApplyVertexColorsToAllLODs(IMeshPaintComponentAdapter& GeometryInfo, USkeletalMeshComponent* SkeletalMeshComponent);
+	void ApplyVertexColorsToAllLODs(IMeshPaintComponentAdapter& GeometryInfo, USkeletalMeshComponent* SkeletalMeshComponent);
 
 	/** Returns the number of Mesh LODs for the given MeshComponent */
-	static int32 GetNumberOfLODs(const UMeshComponent* MeshComponent);
+	int32 GetNumberOfLODs(const UMeshComponent* MeshComponent);
 
 	/** OutNumLODs is set to number of Mesh LODs for the given MeshComponent and returns true, or returns false of given mesh component has no valid LODs */
-	static bool TryGetNumberOfLODs(const UMeshComponent* MeshComponent, int32& OutNumLODs);
+	bool TryGetNumberOfLODs(const UMeshComponent* MeshComponent, int32& OutNumLODs);
 	
 	/** Returns the number of Texture Coordinates for the given MeshComponent */
-	static int32 GetNumberOfUVs(const UMeshComponent* MeshComponent, int32 LODIndex);
+	int32 GetNumberOfUVs(const UMeshComponent* MeshComponent, int32 LODIndex);
 
 	/** Checks whether or not the mesh components contains per lod colors (for all LODs)*/
-	static bool DoesMeshComponentContainPerLODColors(const UMeshComponent* MeshComponent);
+	bool DoesMeshComponentContainPerLODColors(const UMeshComponent* MeshComponent);
 
-	/** Retrieves the number of bytes used to store the per-instance LOD vertex color data from the static mesh component */
-	static void GetInstanceColorDataInfo(const UStaticMeshComponent* StaticMeshComponent, int32 LODIndex, int32& OutTotalInstanceVertexColorBytes);
+	/** Retrieves the number of bytes used to store the per-instance LOD vertex color data from the mesh component */
+	void GetInstanceColorDataInfo(const UStaticMeshComponent* StaticMeshComponent, int32 LODIndex, int32& OutTotalInstanceVertexColorBytes);
 
 	/** Given arguments for an action, and an action - retrieves influences vertices and applies Action to them */
-	static bool ApplyPerVertexPaintAction(FPerVertexPaintActionArgs& InArgs, FPerVertexPaintAction Action);
+	bool ApplyPerVertexPaintAction(FPerVertexPaintActionArgs& InArgs, FPerVertexPaintAction Action);
 
-	static bool GetPerVertexPaintInfluencedVertices(FPerVertexPaintActionArgs& InArgs, TSet<int32>& InfluencedVertices);
+	bool GetPerVertexPaintInfluencedVertices(FPerVertexPaintActionArgs& InArgs, TSet<int32>& InfluencedVertices);
 
 	/** Given the adapter, settings and view-information retrieves influences triangles and applies Action to them */
-	static bool ApplyPerTrianglePaintAction(IMeshPaintComponentAdapter* Adapter, const FVector& CameraPosition, const FVector& HitPosition, const UBrushBaseProperties* Settings, FPerTrianglePaintAction Action, bool bOnlyFrontFacingTriangles);
+	bool ApplyPerTrianglePaintAction(IMeshPaintComponentAdapter* Adapter, const FVector& CameraPosition, const FVector& HitPosition, const UBrushBaseProperties* Settings, FPerTrianglePaintAction Action, bool bOnlyFrontFacingTriangles);
 
 	/** Applies vertex painting to InOutvertexColor according to the given parameters  */
-	static bool PaintVertex(const FVector& InVertexPosition, const FMeshPaintParameters& InParams, FColor& InOutVertexColor);
+	bool PaintVertex(const FVector& InVertexPosition, const FMeshPaintParameters& InParams, FColor& InOutVertexColor);
 
 	/** Applies Vertex Color Painting according to the given parameters */
-	static void ApplyVertexColorPaint(const FMeshPaintParameters &InParams, const FLinearColor &OldColor, FLinearColor &NewColor, const float PaintAmount);
+	void ApplyVertexColorPaint(const FMeshPaintParameters &InParams, const FLinearColor &OldColor, FLinearColor &NewColor, const float PaintAmount);
 
 	/** Applies Vertex Blend Weight Painting according to the given parameters */
-	static void ApplyVertexWeightPaint(const FMeshPaintParameters &InParams, const FLinearColor &OldColor, FLinearColor &NewColor, const float PaintAmount);
+	void ApplyVertexWeightPaint(const FMeshPaintParameters &InParams, const FLinearColor &OldColor, FLinearColor &NewColor, const float PaintAmount);
 
 	/** Generate texture weight color for given number of weights and the to-paint index */
-	static FLinearColor GenerateColorForTextureWeight(const int32 NumWeights, const int32 WeightIndex);
+	FLinearColor GenerateColorForTextureWeight(const int32 NumWeights, const int32 WeightIndex);
 
 	/** Computes the Paint power multiplier value */
-	static float ComputePaintMultiplier(float SquaredDistanceToVertex2D, float BrushStrength, float BrushInnerRadius, float BrushRadialFalloff, float BrushInnerDepth, float BrushDepthFallof, float VertexDepthToBrush);
+	float ComputePaintMultiplier(float SquaredDistanceToVertex2D, float BrushStrength, float BrushInnerRadius, float BrushRadialFalloff, float BrushInnerDepth, float BrushDepthFallof, float VertexDepthToBrush);
 
 	/** Checks whether or not a point is influenced by the painting brush according to the given parameters*/
-	static bool IsPointInfluencedByBrush(const FVector& InPosition, const FMeshPaintParameters& InParams, float& OutSquaredDistanceToVertex2D, float& OutVertexDepthToBrush);
+	bool IsPointInfluencedByBrush(const FVector& InPosition, const FMeshPaintParameters& InParams, float& OutSquaredDistanceToVertex2D, float& OutVertexDepthToBrush);
 
-	static bool IsPointInfluencedByBrush(const FVector2D& BrushSpacePosition, const float BrushRadiusSquared, float& OutInRangeValue);
+	bool IsPointInfluencedByBrush(const FVector2D& BrushSpacePosition, const float BrushRadiusSquared, float& OutInRangeValue);
 
 	template<typename T>
-	static void ApplyBrushToVertex(const FVector& VertexPosition, const FMatrix& InverseBrushMatrix, const float BrushRadius, const float BrushFalloffAmount, const float BrushStrength, const T& PaintValue, T& InOutValue);
-
-public:
-
-
+	void ApplyBrushToVertex(const FVector& VertexPosition, const FMatrix& InverseBrushMatrix, const float BrushRadius, const float BrushFalloffAmount, const float BrushStrength, const T& PaintValue, T& InOutValue);
 
 	/** Helper function to retrieve vertex color from a UTexture given a UVCoordinate */
-	static FColor PickVertexColorFromTextureData(const uint8* MipData, const FVector2D& UVCoordinate, const UTexture2D* Texture, const FColor ColorMask);	
+	FColor PickVertexColorFromTextureData(const uint8* MipData, const FVector2D& UVCoordinate, const UTexture2D* Texture, const FColor ColorMask);	
 };
 
 template<typename T>
-void UMeshPaintingToolset::ApplyBrushToVertex(const FVector& VertexPosition, const FMatrix& InverseBrushMatrix, const float BrushRadius, const float BrushFalloffAmount, const float BrushStrength, const T& PaintValue, T& InOutValue)
+void UMeshPaintingSubsystem::ApplyBrushToVertex(const FVector& VertexPosition, const FMatrix& InverseBrushMatrix, const float BrushRadius, const float BrushFalloffAmount, const float BrushStrength, const T& PaintValue, T& InOutValue)
 {
 	const FVector BrushSpacePosition = InverseBrushMatrix.TransformPosition(VertexPosition);
 	const FVector2D BrushSpacePosition2D(BrushSpacePosition.X, BrushSpacePosition.Y);
@@ -243,7 +240,7 @@ void UMeshPaintingToolset::ApplyBrushToVertex(const FVector& VertexPosition, con
 	if (IsPointInfluencedByBrush(BrushSpacePosition2D, BrushRadius * BrushRadius, InfluencedValue))
 	{
 		float InnerBrushRadius = BrushFalloffAmount * BrushRadius;
-		float PaintStrength = UMeshPaintingToolset::ComputePaintMultiplier(BrushSpacePosition2D.SizeSquared(), BrushStrength, InnerBrushRadius, BrushRadius - InnerBrushRadius, 1.0f, 1.0f, 1.0f);
+		float PaintStrength = GEngine->GetEngineSubsystem<UMeshPaintingSubsystem>()->ComputePaintMultiplier(BrushSpacePosition2D.SizeSquared(), BrushStrength, InnerBrushRadius, BrushRadius - InnerBrushRadius, 1.0f, 1.0f, 1.0f);
 
 		const T OldValue = InOutValue;
 		InOutValue = FMath::LerpStable(OldValue, PaintValue, PaintStrength);
