Moving everything from SharedMeshData to the MeshPaintingSubsystem for one unified data store

diff --git a/Engine/Plugins/MeshPainting/Source/MeshPaintEditorMode/Private/BrushSettingsCustomization.cpp b/Engine/Plugins/MeshPainting/Source/MeshPaintEditorMode/Private/BrushSettingsCustomization.cpp
--- a/Engine/Plugins/MeshPainting/Source/MeshPaintEditorMode/Private/BrushSettingsCustomization.cpp
+++ b/Engine/Plugins/MeshPainting/Source/MeshPaintEditorMode/Private/BrushSettingsCustomization.cpp
@@ -435,7 +435,7 @@ void FTexturePaintingSettingsCustomization::CustomizeDetails(IDetailLayoutBuilde
 			.AllowSpin(true)
 			.Value_Lambda([]() -> int32 { return UMeshPaintMode::GetTextureToolProperties() ? UMeshPaintMode::GetTextureToolProperties()->UVChannel : 0; })
 			.MinValue(0)
-			.MaxValue_Lambda([]() -> int32 { return UMeshPaintMode::GetMeshPaintMode()->GetMeshToolData()->GetMaxUVIndexToPaint(); })
+			.MaxValue_Lambda([]() -> int32 { return GEngine->GetEngineSubsystem<UMeshPaintingSubsystem>()->GetMaxUVIndexToPaint(); })
 			.OnValueChanged(SNumericEntryBox<int32>::FOnValueChanged::CreateLambda([=](int32 Value) { 
 				if (UMeshPaintMode::GetTextureToolProperties())
 				{
diff --git a/Engine/Plugins/MeshPainting/Source/MeshPaintEditorMode/Private/MeshPaintMode.cpp b/Engine/Plugins/MeshPainting/Source/MeshPaintEditorMode/Private/MeshPaintMode.cpp
--- a/Engine/Plugins/MeshPainting/Source/MeshPaintEditorMode/Private/MeshPaintMode.cpp
+++ b/Engine/Plugins/MeshPainting/Source/MeshPaintEditorMode/Private/MeshPaintMode.cpp
@@ -26,6 +26,7 @@
 #include "MeshTexturePaintingTool.h"
 #include "Modules/ModuleManager.h"
 #include "Settings/LevelEditorMiscSettings.h"
+#include "LevelEditor.h"
 
 
 #define LOCTEXT_NAMESPACE "MeshPaintMode"
@@ -137,27 +138,20 @@ void UMeshPaintMode::Enter()
 	
 	FMeshPaintEditorModeCommands ToolManagerCommands = FMeshPaintEditorModeCommands::Get();
 
-	SharedMeshToolData = NewObject<UMeshToolManager>(this);
-
 	UVertexAdapterClickToolBuilder* VertexClickToolBuilder = NewObject<UVertexAdapterClickToolBuilder>(this);
-	VertexClickToolBuilder->SharedMeshToolData = SharedMeshToolData;
 	RegisterTool(ToolManagerCommands.VertexSelect, VertexSelectToolName, VertexClickToolBuilder);
 
 	UTextureAdapterClickToolBuilder* TextureClickToolBuilder = NewObject<UTextureAdapterClickToolBuilder>(this);
-	TextureClickToolBuilder->SharedMeshToolData = SharedMeshToolData;
 	RegisterTool(ToolManagerCommands.TextureSelect, TextureSelectToolName, TextureClickToolBuilder);
 
 
 	UMeshColorPaintingToolBuilder* MeshColorPaintingToolBuilder = NewObject<UMeshColorPaintingToolBuilder>(this);
-	MeshColorPaintingToolBuilder->SharedMeshToolData = SharedMeshToolData;
 	RegisterTool(ToolManagerCommands.ColorPaint, ColorPaintToolName, MeshColorPaintingToolBuilder);
 
 	UMeshWeightPaintingToolBuilder* WeightPaintingToolBuilder = NewObject<UMeshWeightPaintingToolBuilder>(this);
-	WeightPaintingToolBuilder->SharedMeshToolData = SharedMeshToolData;
 	RegisterTool(ToolManagerCommands.WeightPaint, WeightPaintToolName, WeightPaintingToolBuilder);
 
 	UMeshTexturePaintingToolBuilder* TexturePaintingToolBuilder = NewObject<UMeshTexturePaintingToolBuilder>(this);
-	TexturePaintingToolBuilder->SharedMeshToolData = SharedMeshToolData;
 	RegisterTool(ToolManagerCommands.TexturePaint, TexturePaintToolName, TexturePaintingToolBuilder);
 	UpdateSelectedMeshes();
 
@@ -169,6 +163,9 @@ void UMeshPaintMode::Enter()
 	GetToolManager()->ConfigureChangeTrackingMode(EToolChangeTrackingMode::NoChangeTracking);
 	Toolkit->SetCurrentPalette(MeshPaintMode_Color);
 	GetToolManager()->ConfigureChangeTrackingMode(EToolChangeTrackingMode::FullUndoRedo);
+
+	FLevelEditorModule& LevelEditor = FModuleManager::GetModuleChecked<FLevelEditorModule>(FName(TEXT("LevelEditor")));
+	LevelEditor.OnRedrawLevelEditingViewports().AddUObject(this, &UMeshPaintMode::UpdateOnMaterialChange);
 }
 
 void UMeshPaintMode::Exit()
@@ -194,9 +191,10 @@ void UMeshPaintMode::Exit()
 
 	Super::Exit();
 
-	// Cleanup the SharedMeshToolData after the parent class exits, because the tools using it are closed out in UEdMode::Exit
-	SharedMeshToolData->ResetState();
-	SharedMeshToolData = nullptr;
+	GEngine->GetEngineSubsystem<UMeshPaintingSubsystem>()->ResetState();
+
+	FLevelEditorModule& LevelEditor = FModuleManager::GetModuleChecked<FLevelEditorModule>(FName(TEXT("LevelEditor")));
+	LevelEditor.OnRedrawLevelEditingViewports().RemoveAll(this);
 }
 
 void UMeshPaintMode::CreateToolkit()
@@ -264,9 +262,10 @@ void UMeshPaintMode::BindCommands()
 
 	CommandList->MapAction(Commands.CycleToNextLOD, FExecuteAction::CreateUObject(this, &UMeshPaintMode::CycleMeshLODs, 1));
 	CommandList->MapAction(Commands.CycleToPreviousLOD, FExecuteAction::CreateUObject(this, &UMeshPaintMode::CycleMeshLODs, -1));
+	UMeshPaintingSubsystem* MeshPaintingSubsystem = GEngine->GetEngineSubsystem<UMeshPaintingSubsystem>();
 
-	auto IsInValidPaintMode = [this]() -> bool {
-		return (GetVertexToolProperties() != nullptr) && SharedMeshToolData && SharedMeshToolData->SelectionContainsValidAdapters();
+	auto IsInValidPaintMode = [MeshPaintingSubsystem]() -> bool {
+		return (GetVertexToolProperties() != nullptr) && MeshPaintingSubsystem && MeshPaintingSubsystem->SelectionContainsValidAdapters();
 	};
 	CommandList->MapAction(Commands.Fill, 
 		FUIAction(FExecuteAction::CreateUObject(this, &UMeshPaintMode::FillWithVertexColor),
@@ -278,8 +277,8 @@ void UMeshPaintMode::BindCommands()
 			FCanExecuteAction::CreateUObject(this, &UMeshPaintMode::CanPropagateVertexColors)
 		));
 
-	auto IsAValidMeshComponentSelected = [this]() -> bool {
-		return (GetSelectedComponents<UMeshComponent>().Num() == 1) && SharedMeshToolData && SharedMeshToolData->SelectionContainsValidAdapters();
+	auto IsAValidMeshComponentSelected = [this, MeshPaintingSubsystem]() -> bool {
+		return (GetSelectedComponents<UMeshComponent>().Num() == 1) && MeshPaintingSubsystem && MeshPaintingSubsystem->SelectionContainsValidAdapters();
 	};
 	CommandList->MapAction(Commands.Import,
 		FUIAction(
@@ -360,10 +359,10 @@ void UMeshPaintMode::OnVertexPaintFinished()
 		}
 		else
 		{
-			if (SharedMeshToolData)
+			if (UMeshPaintingSubsystem* MeshPaintingSubsystem = GEngine->GetEngineSubsystem<UMeshPaintingSubsystem>())
 			{
 				UpdateCachedVertexDataSize();
-				SharedMeshToolData->Refresh();
+				MeshPaintingSubsystem->Refresh();
 			}
 		}
 	}
@@ -383,28 +382,76 @@ void UMeshPaintMode::ActorSelectionChangeNotify()
 
 void UMeshPaintMode::UpdateSelectedMeshes()
 {
-	if (SharedMeshToolData)
+	if (UMeshPaintingSubsystem* MeshPaintingSubsystem = GEngine->GetEngineSubsystem<UMeshPaintingSubsystem>())
 	{
-		SharedMeshToolData->ResetState();
+		FName PaletteName = Toolkit->GetCurrentPalette();
+
+		MeshPaintingSubsystem->ResetState();
 		const TArray<UMeshComponent*> CurrentMeshComponents = GetSelectedComponents<UMeshComponent>();
-		SharedMeshToolData->AddSelectedMeshComponents(CurrentMeshComponents);
-		SharedMeshToolData->bNeedsRecache = true;
+		MeshPaintingSubsystem->AddSelectedMeshComponents(CurrentMeshComponents);
+		MeshPaintingSubsystem->bNeedsRecache = true;
+
+
+		// Check the current selection for a material that supports texture paint
+		if (PaletteName == UMeshPaintMode::MeshPaintMode_Texture)
+		{
+			CheckSelectionForTexturePaintCompat(CurrentMeshComponents);
+		}
 	}
+	
 	bRecacheVertexDataSize = true;
 }
 
+void UMeshPaintMode::CheckSelectionForTexturePaintCompat(const TArray<UMeshComponent*>& CurrentMeshComponents)
+{
+	bool bCurrentSelectionSupportsTexturePaint = false;
+	if (CurrentMeshComponents.Num())
+	{
+		UActorComponent* ActorComponent = CurrentMeshComponents[0];
+		TArray<FPaintableTexture> PaintableTextures;
+		if (ActorComponent)
+		{
+			UMeshComponent* MeshComponent = Cast<UMeshComponent>(ActorComponent);
+			if (MeshComponent)
+			{
+				TSharedPtr<IMeshPaintComponentAdapter> MeshAdapter = GEngine->GetEngineSubsystem<UMeshPaintingSubsystem>()->GetAdapterForComponent(MeshComponent);
+				if (!MeshAdapter.IsValid())
+				{
+					MeshAdapter = FMeshPaintComponentAdapterFactory::CreateAdapterForMesh(MeshComponent, 0);
+				}
+				UTexturePaintToolset::RetrieveTexturesForComponent(MeshComponent, MeshAdapter.Get(), PaintableTextures);
+			}
+		}
+		bCurrentSelectionSupportsTexturePaint = PaintableTextures.Num() > 0;
+	}
+
+	GEngine->GetEngineSubsystem<UMeshPaintingSubsystem>()->SetSelectionHasMaterialValidForTexturePaint(bCurrentSelectionSupportsTexturePaint);
+}
+
+void UMeshPaintMode::UpdateOnMaterialChange(bool bInvalidateHitProxies)
+{
+	FName PaletteName = Toolkit->GetCurrentPalette();
+	// Check the current selection for a material that supports texture paint
+	if (PaletteName == UMeshPaintMode::MeshPaintMode_Texture)
+	{
+		const TArray<UMeshComponent*> CurrentMeshComponents = GetSelectedComponents<UMeshComponent>();
+		CheckSelectionForTexturePaintCompat(CurrentMeshComponents);
+	}
+}
+
+
 void UMeshPaintMode::FillWithVertexColor()
 {
 	FScopedTransaction Transaction(LOCTEXT("LevelMeshPainter_TransactionFillInstColors", "Filling Per-Instance Vertex Colors"));
 	const TArray<UMeshComponent*> MeshComponents = GetSelectedComponents<UMeshComponent>();
-
+	UMeshPaintingSubsystem* MeshPaintingSubsystem = GEngine->GetEngineSubsystem<UMeshPaintingSubsystem>();
 	static const bool bConvertSRGB = false;
 	FColor FillColor = FColor::White;
 	FColor MaskColor = FColor::White;
 
 	if (GetToolManager()->GetActiveTool(EToolSide::Mouse)->IsA<UMeshWeightPaintingTool>())
 	{
-		FillColor = GEngine->GetEngineSubsystem<UMeshPaintingSubsystem>()->GenerateColorForTextureWeight((int32)GetWeightToolProperties()->TextureWeightType, (int32)GetWeightToolProperties()->PaintTextureWeightIndex).ToFColor(bConvertSRGB);
+		FillColor = MeshPaintingSubsystem->GenerateColorForTextureWeight((int32)GetWeightToolProperties()->TextureWeightType, (int32)GetWeightToolProperties()->PaintTextureWeightIndex).ToFColor(bConvertSRGB);
 	}
 	else if (UMeshColorPaintingToolProperties* ColorProperties = GetColorToolProperties())
 	{
@@ -423,7 +470,7 @@ void UMeshPaintMode::FillWithVertexColor()
 		Component->Modify();
 		ComponentReregisterContext = MakeUnique<FComponentReregisterContext>(Component);
 
-		TSharedPtr<IMeshPaintComponentAdapter> MeshAdapter = SharedMeshToolData->GetAdapterForComponent(Component);
+		TSharedPtr<IMeshPaintComponentAdapter> MeshAdapter = MeshPaintingSubsystem->GetAdapterForComponent(Component);
 		if (MeshAdapter)
 		{
 			MeshAdapter->PreEdit();
@@ -435,7 +482,7 @@ void UMeshPaintMode::FillWithVertexColor()
 
 		if (Component->IsA<UStaticMeshComponent>())
 		{
-			GEngine->GetEngineSubsystem<UMeshPaintingSubsystem>()->FillStaticMeshVertexColors(Cast<UStaticMeshComponent>(Component), bPaintOnSpecificLOD ? ColorProperties->LODIndex : -1, FillColor, MaskColor);
+			MeshPaintingSubsystem->FillStaticMeshVertexColors(Cast<UStaticMeshComponent>(Component), bPaintOnSpecificLOD ? ColorProperties->LODIndex : -1, FillColor, MaskColor);
 		}
 		else if (Component->IsA<USkeletalMeshComponent>())
 		{
@@ -581,7 +628,7 @@ bool UMeshPaintMode::CanRemoveInstanceColors() const
 bool UMeshPaintMode::CanPasteInstanceVertexColors() const
 {
 	const TArray<UStaticMeshComponent*> StaticMeshComponents = GetSelectedComponents<UStaticMeshComponent>();
-	const TArray<FPerComponentVertexColorData> CopiedColorsByComponent = SharedMeshToolData->GetCopiedColorsByComponent();
+	const TArray<FPerComponentVertexColorData> CopiedColorsByComponent = GEngine->GetEngineSubsystem<UMeshPaintingSubsystem>()->GetCopiedColorsByComponent();
 	return GEditor->GetEditorSubsystem<UMeshPaintModeSubsystem>()->CanPasteInstanceVertexColors(StaticMeshComponents, CopiedColorsByComponent);
 }
 
@@ -606,7 +653,7 @@ bool UMeshPaintMode::CanPropagateVertexColorsToLODs() const
 		bPaintOnSpecificLOD = ColorProperties ? ColorProperties->bPaintOnSpecificLOD : false;
 	}
 	// Can propagate when the mesh contains per-lod vertex colors or when we are not painting to a specific lod
-	const bool bSelectionContainsPerLODColors = SharedMeshToolData->SelectionContainsPerLODColors();
+	const bool bSelectionContainsPerLODColors = GEngine->GetEngineSubsystem<UMeshPaintingSubsystem>()->SelectionContainsPerLODColors();
 	return bSelectionContainsPerLODColors || !bPaintOnSpecificLOD;
 }
 
@@ -615,7 +662,7 @@ void UMeshPaintMode::CopyVertexColors()
 	const TArray<UStaticMeshComponent*> StaticMeshComponents = GetSelectedComponents<UStaticMeshComponent>();
 	TArray<FPerComponentVertexColorData> CopiedColorsByComponent;
 	GEditor->GetEditorSubsystem<UMeshPaintModeSubsystem>()->CopyVertexColors(StaticMeshComponents, CopiedColorsByComponent);
-	SharedMeshToolData->SetCopiedColorsByComponent(CopiedColorsByComponent);
+	GEngine->GetEngineSubsystem<UMeshPaintingSubsystem>()->SetCopiedColorsByComponent(CopiedColorsByComponent);
 
 }
 
@@ -623,7 +670,7 @@ void UMeshPaintMode::PasteVertexColors()
 {
 	FScopedTransaction Transaction(LOCTEXT("LevelMeshPainter_TransactionPasteInstColors", "Pasting Per-Instance Vertex Colors"));
 	const TArray<UStaticMeshComponent*> StaticMeshComponents = GetSelectedComponents<UStaticMeshComponent>();
-	TArray<FPerComponentVertexColorData> CopiedColorsByComponent = SharedMeshToolData->GetCopiedColorsByComponent();
+	TArray<FPerComponentVertexColorData> CopiedColorsByComponent = GEngine->GetEngineSubsystem<UMeshPaintingSubsystem>()->GetCopiedColorsByComponent();
 	GEditor->GetEditorSubsystem<UMeshPaintModeSubsystem>()->PasteVertexColors(StaticMeshComponents, CopiedColorsByComponent);
 	UpdateCachedVertexDataSize();
 }
@@ -668,8 +715,9 @@ void UMeshPaintMode::PropagateVertexColorsToLODs()
 {
 	//Only show the lost data warning if there is actually some data to lose
 	bool bAbortChange = false;
-	TArray<UMeshComponent*> PaintableComponents = SharedMeshToolData->GetPaintableMeshComponents();
-	const bool bSelectionContainsPerLODColors = SharedMeshToolData->SelectionContainsPerLODColors();
+	UMeshPaintingSubsystem* MeshPaintingSubsystem = GEngine->GetEngineSubsystem<UMeshPaintingSubsystem>();
+	TArray<UMeshComponent*> PaintableComponents = MeshPaintingSubsystem->GetPaintableMeshComponents();
+	const bool bSelectionContainsPerLODColors = MeshPaintingSubsystem->SelectionContainsPerLODColors();
 	if (bSelectionContainsPerLODColors)
 	{
 		//Warn the user they will lose custom painting data
@@ -690,7 +738,7 @@ void UMeshPaintMode::PropagateVertexColorsToLODs()
 		else
 		{
 			// Reset the state flag as we'll be removing all per-lod colors 
-			SharedMeshToolData->ClearSelectionLODColors();
+			MeshPaintingSubsystem->ClearSelectionLODColors();
 			GEditor->GetEditorSubsystem<UMeshPaintModeSubsystem>()->RemovePerLODColors(PaintableComponents);
 		}
 	}
@@ -705,13 +753,13 @@ void UMeshPaintMode::PropagateVertexColorsToLODs()
 	{
 		if (SelectedComponent)
 		{
-			TSharedPtr<IMeshPaintComponentAdapter> MeshAdapter = SharedMeshToolData->GetAdapterForComponent(SelectedComponent);
+			TSharedPtr<IMeshPaintComponentAdapter> MeshAdapter = MeshPaintingSubsystem->GetAdapterForComponent(SelectedComponent);
 			GEngine->GetEngineSubsystem<UMeshPaintingSubsystem>()->ApplyVertexColorsToAllLODs(*MeshAdapter, SelectedComponent);
 			FComponentReregisterContext ReregisterContext(SelectedComponent);
 		}
 	}
 	UpdateCachedVertexDataSize();
-	SharedMeshToolData->Refresh();
+	MeshPaintingSubsystem->Refresh();
 }
 
 template<typename ComponentClass>
@@ -757,26 +805,20 @@ template TArray<UStaticMeshComponent*> UMeshPaintMode::GetSelectedComponents<USt
 template TArray<USkeletalMeshComponent*> UMeshPaintMode::GetSelectedComponents<USkeletalMeshComponent>() const;
 template TArray<UMeshComponent*> UMeshPaintMode::GetSelectedComponents<UMeshComponent>() const;
 
-TWeakObjectPtr<UMeshToolManager> UMeshPaintMode::GetMeshToolData() const
-{
-	return SharedMeshToolData;
-}
-
-
 
 void UMeshPaintMode::UpdateCachedVertexDataSize()
 {
 	CachedVertexDataSize = 0;
 
 	const bool bInstance = true;
-	if (SharedMeshToolData)
+	if (UMeshPaintingSubsystem* MeshPaintingSubsystem = GEngine->GetEngineSubsystem<UMeshPaintingSubsystem>())
 	{
-		for (UMeshComponent* SelectedComponent : SharedMeshToolData->GetPaintableMeshComponents())
+		for (UMeshComponent* SelectedComponent : MeshPaintingSubsystem->GetPaintableMeshComponents())
 		{
-			int32 NumLODs = GEngine->GetEngineSubsystem<UMeshPaintingSubsystem>()->GetNumberOfLODs(SelectedComponent);
+			int32 NumLODs = MeshPaintingSubsystem->GetNumberOfLODs(SelectedComponent);
 			for (int32 LODIndex = 0; LODIndex < NumLODs; ++LODIndex)
 			{
-				CachedVertexDataSize += GEngine->GetEngineSubsystem<UMeshPaintingSubsystem>()->GetVertexColorBufferSize(SelectedComponent, LODIndex, bInstance);
+				CachedVertexDataSize += MeshPaintingSubsystem->GetVertexColorBufferSize(SelectedComponent, LODIndex, bInstance);
 			}
 		}
 	}
diff --git a/Engine/Plugins/MeshPainting/Source/MeshPaintEditorMode/Private/MeshPaintMode.h b/Engine/Plugins/MeshPainting/Source/MeshPaintEditorMode/Private/MeshPaintMode.h
--- a/Engine/Plugins/MeshPainting/Source/MeshPaintEditorMode/Private/MeshPaintMode.h
+++ b/Engine/Plugins/MeshPainting/Source/MeshPaintEditorMode/Private/MeshPaintMode.h
@@ -63,8 +63,6 @@ public:
 		return CachedVertexDataSize;
 	}
 
-	TWeakObjectPtr<UMeshToolManager> GetMeshToolData() const;
-
 protected:
 
 	/** Binds UI commands to actions for the mesh paint mode */
@@ -78,6 +76,11 @@ protected:
 	virtual void UpdateOnPaletteChange(FName NewPalette);
 	// end UEdMode Interface
 	void UpdateSelectedMeshes();
+
+	void CheckSelectionForTexturePaintCompat(const TArray<UMeshComponent*>& CurrentMeshComponents);
+	
+
+	void UpdateOnMaterialChange(bool bInvalidateHitProxies);
 	void OnResetViewMode();
 	void FillWithVertexColor();
 	void PropagateVertexColorsToAsset();
@@ -108,8 +111,6 @@ protected:
 	UPROPERTY(Transient)
 	UMeshPaintModeSettings* ModeSettings;
 
-	UPROPERTY(Transient)
-	UMeshToolManager* SharedMeshToolData;
 
 	// End vertex paint state
 	FGetSelectedMeshComponents MeshComponentDelegate;
diff --git a/Engine/Plugins/MeshPainting/Source/MeshPaintingToolset/Private/MeshPaintHelpers.cpp b/Engine/Plugins/MeshPainting/Source/MeshPaintingToolset/Private/MeshPaintHelpers.cpp
--- a/Engine/Plugins/MeshPainting/Source/MeshPaintingToolset/Private/MeshPaintHelpers.cpp
+++ b/Engine/Plugins/MeshPainting/Source/MeshPaintingToolset/Private/MeshPaintHelpers.cpp
@@ -27,6 +27,13 @@
 
 extern void PropagateVertexPaintToSkeletalMesh(USkeletalMesh* SkeletalMesh, int32 LODIndex);
 
+UMeshPaintingSubsystem::UMeshPaintingSubsystem()
+	:	bNeedsRecache(true),
+		bSelectionHasMaterialValidForTexturePaint(false)
+{
+
+}
+
 bool UMeshPaintingSubsystem::HasPaintableMesh(UActorComponent* Component)
 {
 	return Cast<UMeshComponent>(Component) != nullptr;
@@ -1126,6 +1133,7 @@ FColor UMeshPaintingSubsystem::PickVertexColorFromTextureData(const uint8* MipDa
 	return VertexColor;
 }
 
+
 bool UMeshPaintingSubsystem::GetPerVertexPaintInfluencedVertices(FPerVertexPaintActionArgs& InArgs, TSet<int32>& InfluencedVertices)
 {
 	// Retrieve components world matrix
@@ -1388,69 +1396,63 @@ void UMeshPaintingSubsystem::ApplyVertexColorsToAllLODs(IMeshPaintComponentAdapt
 		}
 	}
 }
-
-UMeshToolManager::UMeshToolManager()
-	:Super()
-{
-	bNeedsRecache = true;
-}
-
-TMap<UMeshComponent*, TSharedPtr<IMeshPaintComponentAdapter>> UMeshToolManager::GetComponentToAdapterMap() const
+TMap<FName, TSharedPtr<IMeshPaintComponentAdapter>> UMeshPaintingSubsystem::GetComponentToAdapterMap() const
 {
 	return ComponentToAdapterMap;
 }
 
-TSharedPtr<IMeshPaintComponentAdapter> UMeshToolManager::GetAdapterForComponent(const UMeshComponent* InComponent)
+TSharedPtr<IMeshPaintComponentAdapter> UMeshPaintingSubsystem::GetAdapterForComponent(const UMeshComponent* InComponent)
 {
-	TSharedPtr<IMeshPaintComponentAdapter>* MeshAdapterPtr = ComponentToAdapterMap.Find(InComponent);
+	FName ComponentName = InComponent->GetFName();
+	TSharedPtr<IMeshPaintComponentAdapter>* MeshAdapterPtr = ComponentToAdapterMap.Find(ComponentName);
 	return MeshAdapterPtr ? *MeshAdapterPtr : TSharedPtr<IMeshPaintComponentAdapter>();
 }
 
-void UMeshToolManager::AddToComponentToAdapterMap(UMeshComponent* InComponent, TSharedPtr<IMeshPaintComponentAdapter> InAdapter)
+void UMeshPaintingSubsystem::AddToComponentToAdapterMap(UMeshComponent* InComponent, TSharedPtr<IMeshPaintComponentAdapter> InAdapter)
 {
-	ComponentToAdapterMap.Add(InComponent, InAdapter);
+	ComponentToAdapterMap.Add(InComponent->GetFName(), InAdapter);
 }
 
-TArray<UMeshComponent*> UMeshToolManager::GetSelectedMeshComponents() const
+TArray<UMeshComponent*> UMeshPaintingSubsystem::GetSelectedMeshComponents() const
 {
 	return SelectedMeshComponents;
 }
 
-void UMeshToolManager::ClearSelectedMeshComponents()
+void UMeshPaintingSubsystem::ClearSelectedMeshComponents()
 {
 	SelectedMeshComponents.Empty();
 }
 
 
-void UMeshToolManager::AddSelectedMeshComponents(const TArray<UMeshComponent*>& InComponents)
+void UMeshPaintingSubsystem::AddSelectedMeshComponents(const TArray<UMeshComponent*>& InComponents)
 {
 	SelectedMeshComponents.Append(InComponents);
 }
 
 
-TArray<UMeshComponent*> UMeshToolManager::GetPaintableMeshComponents() const
+TArray<UMeshComponent*> UMeshPaintingSubsystem::GetPaintableMeshComponents() const
 {
 	return PaintableComponents;
 }
 
-void UMeshToolManager::AddPaintableMeshComponent(UMeshComponent* InComponent)
+void UMeshPaintingSubsystem::AddPaintableMeshComponent(UMeshComponent* InComponent)
 {
 	PaintableComponents.Add(InComponent);
 }
 
-void UMeshToolManager::ClearPaintableMeshComponents()
+void UMeshPaintingSubsystem::ClearPaintableMeshComponents()
 {
 	PaintableComponents.Empty();
 }
 
-void UMeshToolManager::ResetState()
+void UMeshPaintingSubsystem::ResetState()
 {
 	PaintableComponents.Empty();
 	ComponentToAdapterMap.Empty();
 	SelectedMeshComponents.Empty();
 }
 
-void UMeshToolManager::Refresh()
+void UMeshPaintingSubsystem::Refresh()
 {
 	// Ensure that we call OnRemoved while adapter/components are still valid
 	PaintableComponents.Empty();
@@ -1459,7 +1461,7 @@ void UMeshToolManager::Refresh()
 	bNeedsRecache = true;
 }
 
-void UMeshToolManager::CleanUp()
+void UMeshPaintingSubsystem::CleanUp()
 {
 	for (auto MeshAdapterPair : ComponentToAdapterMap)
 	{
@@ -1469,7 +1471,7 @@ void UMeshToolManager::CleanUp()
 	FMeshPaintComponentAdapterFactory::CleanupGlobals();
 }
 
-bool UMeshToolManager::FindHitResult(const FRay Ray, FHitResult& BestTraceResult)
+bool UMeshPaintingSubsystem::FindHitResult(const FRay Ray, FHitResult& BestTraceResult)
 {
 	const FVector& Origin = Ray.Origin;
 	const FVector& Direction = Ray.Direction;
@@ -1499,7 +1501,7 @@ bool UMeshToolManager::FindHitResult(const FRay Ray, FHitResult& BestTraceResult
 	return true;
 }
 
-bool UMeshToolManager::SelectionContainsValidAdapters() const
+bool UMeshPaintingSubsystem::SelectionContainsValidAdapters() const
 {
 	for (auto& MeshAdapterPair : ComponentToAdapterMap)
 	{
@@ -1512,17 +1514,17 @@ bool UMeshToolManager::SelectionContainsValidAdapters() const
 	return false;
 }
 
-TArray<FPerComponentVertexColorData> UMeshToolManager::GetCopiedColorsByComponent() const
+TArray<FPerComponentVertexColorData> UMeshPaintingSubsystem::GetCopiedColorsByComponent() const
 {
 	return CopiedColorsByComponent;
 }
 
-void UMeshToolManager::SetCopiedColorsByComponent(TArray<FPerComponentVertexColorData>& InCopiedColors)
+void UMeshPaintingSubsystem::SetCopiedColorsByComponent(TArray<FPerComponentVertexColorData>& InCopiedColors)
 {
 	CopiedColorsByComponent = InCopiedColors;
 }
 
-void UMeshToolManager::CacheSelectionData(const int32 PaintLODIndex, const int32 UVChannel)
+void UMeshPaintingSubsystem::CacheSelectionData(const int32 PaintLODIndex, const int32 UVChannel)
 {
 	bSelectionContainsPerLODColors = false;
 	for (UMeshComponent* MeshComponent : SelectedMeshComponents)
@@ -1542,7 +1544,7 @@ void UMeshToolManager::CacheSelectionData(const int32 PaintLODIndex, const int32
 	bNeedsRecache = false;
 }
 
-int32 UMeshToolManager::GetMaxUVIndexToPaint() const
+int32 UMeshPaintingSubsystem::GetMaxUVIndexToPaint() const
 {
 	if (PaintableComponents.Num() == 1 && PaintableComponents[0])
 	{
diff --git a/Engine/Plugins/MeshPainting/Source/MeshPaintingToolset/Private/MeshPaintInteractions.cpp b/Engine/Plugins/MeshPainting/Source/MeshPaintingToolset/Private/MeshPaintInteractions.cpp
--- a/Engine/Plugins/MeshPainting/Source/MeshPaintingToolset/Private/MeshPaintInteractions.cpp
+++ b/Engine/Plugins/MeshPainting/Source/MeshPaintingToolset/Private/MeshPaintInteractions.cpp
@@ -30,15 +30,15 @@ FInputRayHit UMeshPaintSelectionMechanic::IsHitByClick(const FInputDeviceRay& Cl
 
 void UMeshPaintSelectionMechanic::OnClicked(const FInputDeviceRay& ClickPos)
 {
-	if (SharedMeshToolData.IsValid())
+	if (UMeshPaintingSubsystem* MeshPaintingSubsystem = GEngine->GetEngineSubsystem<UMeshPaintingSubsystem>())
 	{
 		for (UMeshComponent* MeshComponent : CachedClickedComponents)
 		{
-			TSharedPtr<IMeshPaintComponentAdapter> MeshAdapter = SharedMeshToolData->GetAdapterForComponent(MeshComponent);
+			TSharedPtr<IMeshPaintComponentAdapter> MeshAdapter = MeshPaintingSubsystem->GetAdapterForComponent(MeshComponent);
 			IMeshPaintSelectionInterface* Interface = Cast<IMeshPaintSelectionInterface>(GetParentTool());
 			if (MeshComponent && MeshComponent->IsVisible() && MeshAdapter.IsValid() && MeshAdapter->IsValid() && Interface->IsMeshAdapterSupported(MeshAdapter))
 			{
-				SharedMeshToolData->AddPaintableMeshComponent(MeshComponent);
+				MeshPaintingSubsystem->AddPaintableMeshComponent(MeshComponent);
 				MeshAdapter->OnAdded();
 			}
 
@@ -59,7 +59,8 @@ bool UMeshPaintSelectionMechanic::FindClickedComponentsAndCacheAdapters(const FI
 {
 	bool bFoundValidComponents = false;
 #if WITH_EDITOR
-	if (!SharedMeshToolData.IsValid())
+	UMeshPaintingSubsystem* MeshPaintingSubsystem = GEngine->GetEngineSubsystem<UMeshPaintingSubsystem>();
+	if (!MeshPaintingSubsystem)
 	{
 		return bFoundValidComponents;
 	}
@@ -74,11 +75,15 @@ bool UMeshPaintSelectionMechanic::FindClickedComponentsAndCacheAdapters(const FI
 			for (UActorComponent* CandidateComponent : CandidateComponents)
 			{
 				UMeshComponent* MeshComponent = Cast<UMeshComponent>(CandidateComponent);
-				TSharedPtr<IMeshPaintComponentAdapter> MeshAdapter = FMeshPaintComponentAdapterFactory::CreateAdapterForMesh(MeshComponent, 0);
+				TSharedPtr<IMeshPaintComponentAdapter> MeshAdapter = MeshPaintingSubsystem->GetAdapterForComponent(MeshComponent);
+				if (!MeshAdapter.IsValid())
+				{
+					MeshAdapter = FMeshPaintComponentAdapterFactory::CreateAdapterForMesh(MeshComponent, 0);
+					MeshPaintingSubsystem->AddToComponentToAdapterMap(MeshComponent, MeshAdapter);
+				}
 				IMeshPaintSelectionInterface* Interface = Cast<IMeshPaintSelectionInterface>(GetParentTool());
 				if (MeshAdapter.IsValid() && Interface->IsMeshAdapterSupported(MeshAdapter))
 				{
-					SharedMeshToolData->AddToComponentToAdapterMap(MeshComponent, MeshAdapter);
 					CachedClickedComponents.AddUnique(MeshComponent);
 					CachedClickedActors.AddUnique(Cast<AActor>(MeshComponent->GetOuter()));
 					bFoundValidComponents = true;
@@ -90,9 +95,5 @@ bool UMeshPaintSelectionMechanic::FindClickedComponentsAndCacheAdapters(const FI
 	return bFoundValidComponents;
 }
 
-void UMeshPaintSelectionMechanic::SetMeshToolData(TWeakObjectPtr<UMeshToolManager> InMeshToolData)
-{
-	SharedMeshToolData = InMeshToolData;
-}
 
 #undef LOCTEXT_NAMESPACE
\ No newline at end of file
diff --git a/Engine/Plugins/MeshPainting/Source/MeshPaintingToolset/Private/MeshSelect.cpp b/Engine/Plugins/MeshPainting/Source/MeshPaintingToolset/Private/MeshSelect.cpp
--- a/Engine/Plugins/MeshPainting/Source/MeshPaintingToolset/Private/MeshSelect.cpp
+++ b/Engine/Plugins/MeshPainting/Source/MeshPaintingToolset/Private/MeshSelect.cpp
@@ -22,7 +22,6 @@ bool UVertexAdapterClickToolBuilder::CanBuildTool(const FToolBuilderState& Scene
 UInteractiveTool* UVertexAdapterClickToolBuilder::BuildTool(const FToolBuilderState& SceneState) const
 {
 	UVertexAdapterClickTool* NewTool = NewObject<UVertexAdapterClickTool>(SceneState.ToolManager);
-	NewTool->SetMeshToolData(SharedMeshToolData);
 	return NewTool;
 }
 
@@ -34,8 +33,6 @@ bool UTextureAdapterClickToolBuilder::CanBuildTool(const FToolBuilderState& Scen
 UInteractiveTool* UTextureAdapterClickToolBuilder::BuildTool(const FToolBuilderState& SceneState) const
 {
 	UTextureAdapterClickTool* NewTool = NewObject<UTextureAdapterClickTool>(SceneState.ToolManager);
-	NewTool->SetMeshToolData(SharedMeshToolData);
-
 	return NewTool;
 }
 
@@ -61,7 +58,6 @@ void UMeshClickTool::Setup()
 	// Set up selection mechanic to select valid meshes
 	SelectionMechanic = NewObject<UMeshPaintSelectionMechanic>(this);
 	SelectionMechanic->Setup(this);
-	SelectionMechanic->SetMeshToolData(SharedMeshToolData);
 }
 
 void UMeshClickTool::OnUpdateModifierState(int ModifierID, bool bIsOn)
@@ -82,11 +78,6 @@ void UMeshClickTool::OnClicked(const FInputDeviceRay& ClickPos)
 	SelectionMechanic->OnClicked(ClickPos);
 }
 
-void UMeshClickTool::SetMeshToolData(TWeakObjectPtr<UMeshToolManager> InMeshToolData)
-{
-	SharedMeshToolData = InMeshToolData;
-}
-
 
 UVertexAdapterClickTool::UVertexAdapterClickTool()
 	: UMeshClickTool()
diff --git a/Engine/Plugins/MeshPainting/Source/MeshPaintingToolset/Private/MeshTexturePaintingTool.cpp b/Engine/Plugins/MeshPainting/Source/MeshPaintingToolset/Private/MeshTexturePaintingTool.cpp
--- a/Engine/Plugins/MeshPainting/Source/MeshPaintingToolset/Private/MeshTexturePaintingTool.cpp
+++ b/Engine/Plugins/MeshPainting/Source/MeshPaintingToolset/Private/MeshTexturePaintingTool.cpp
@@ -30,29 +30,14 @@
 
 bool UMeshTexturePaintingToolBuilder::CanBuildTool(const FToolBuilderState& SceneState) const
 {
-	const TFunction<bool(UActorComponent*)>& HasPaintableMesh = [this](UActorComponent* InComponent)
-	{ 
-		return GEngine->GetEngineSubsystem<UMeshPaintingSubsystem>()->HasPaintableMesh(InComponent);
-	};
-	UActorComponent* ActorComponent = ToolBuilderUtil::FindFirstComponent(SceneState, HasPaintableMesh);
-	TArray<FPaintableTexture> PaintableTextures;
-	if (ActorComponent && SharedMeshToolData.IsValid())
-	{
-		UMeshComponent* MeshComponent = Cast<UMeshComponent>(ActorComponent);
-		if (MeshComponent)
-		{
-			TSharedPtr<IMeshPaintComponentAdapter> MeshAdapter = FMeshPaintComponentAdapterFactory::CreateAdapterForMesh(MeshComponent, 0);
-			UTexturePaintToolset::RetrieveTexturesForComponent(MeshComponent, MeshAdapter.Get(), PaintableTextures);
-		}
-	}
 
-	return PaintableTextures.Num() > 0;
+
+	return GEngine->GetEngineSubsystem<UMeshPaintingSubsystem>()->SelectionHasMaterialValidForTexturePaint();
 }
 
 UInteractiveTool* UMeshTexturePaintingToolBuilder::BuildTool(const FToolBuilderState& SceneState) const
 {
 	UMeshTexturePaintingTool* NewTool = NewObject<UMeshTexturePaintingTool>(SceneState.ToolManager);
-	NewTool->SetMeshToolData(SharedMeshToolData);
 	return NewTool;
 }
 
@@ -113,9 +98,9 @@ void UMeshTexturePaintingTool::Shutdown(EToolShutdownType ShutdownType)
 
 	// Remove any existing texture targets
 	TexturePaintTargetList.Empty();
-	if (SharedMeshToolData.IsValid())
+	if (UMeshPaintingSubsystem* MeshPaintingSubsystem = GEngine->GetEngineSubsystem<UMeshPaintingSubsystem>())
 	{
-		SharedMeshToolData->Refresh();
+		MeshPaintingSubsystem->Refresh();
 	}
 
 	BrushProperties->SaveProperties(this);
@@ -127,8 +112,8 @@ void UMeshTexturePaintingTool::Render(IToolsContextRenderAPI* RenderAPI)
 	Super::Render(RenderAPI);
 	FToolDataVisualizer Draw;
 	Draw.BeginFrame(RenderAPI);
-
-	if (SharedMeshToolData.IsValid() && LastBestHitResult.Component != nullptr)
+	UMeshPaintingSubsystem* MeshPaintingSubsystem = GEngine->GetEngineSubsystem<UMeshPaintingSubsystem>();
+	if (MeshPaintingSubsystem && LastBestHitResult.Component != nullptr)
 	{
 		BrushStampIndicator->bDrawIndicatorLines = true;
 		static float WidgetLineThickness = 1.0f;
@@ -142,9 +127,9 @@ void UMeshTexturePaintingTool::Render(IToolsContextRenderAPI* RenderAPI)
 		const FVector NormalLineEnd(LastBestHitResult.Location + LastBestHitResult.Normal * NormalLineSize);
 		Draw.DrawLine(FVector(LastBestHitResult.Location), NormalLineEnd, NormalLineColor, WidgetLineThickness);
 
-		for (UMeshComponent* CurrentComponent : SharedMeshToolData->GetPaintableMeshComponents())
+		for (UMeshComponent* CurrentComponent : MeshPaintingSubsystem->GetPaintableMeshComponents())
 		{
-			TSharedPtr<IMeshPaintComponentAdapter> MeshAdapter = SharedMeshToolData->GetAdapterForComponent(Cast<UMeshComponent>(CurrentComponent));
+			TSharedPtr<IMeshPaintComponentAdapter> MeshAdapter = MeshPaintingSubsystem->GetAdapterForComponent(Cast<UMeshComponent>(CurrentComponent));
 
 			if (MeshAdapter->IsValid() && MeshAdapter->SupportsVertexPaint())
 			{
@@ -171,9 +156,10 @@ void UMeshTexturePaintingTool::Render(IToolsContextRenderAPI* RenderAPI)
 
 void UMeshTexturePaintingTool::OnTick(float DeltaTime)
 {
-	if (SharedMeshToolData.IsValid())
+	UMeshPaintingSubsystem* MeshPaintingSubsystem = GEngine->GetEngineSubsystem<UMeshPaintingSubsystem>();
+	if (MeshPaintingSubsystem)
 	{
-		if (SharedMeshToolData->bNeedsRecache)
+		if (MeshPaintingSubsystem->bNeedsRecache)
 		{
 			CacheSelectionData();
 			CacheTexturePaintData();
@@ -205,12 +191,13 @@ double UMeshTexturePaintingTool::EstimateMaximumTargetDimension()
 {
 	bool bFoundComponentToUse = false;
 	FBoxSphereBounds Bounds = FBoxSphereBounds(0.0);
-	if (SharedMeshToolData.IsValid())
+	UMeshPaintingSubsystem* MeshPaintingSubsystem = GEngine->GetEngineSubsystem<UMeshPaintingSubsystem>();
+	if (MeshPaintingSubsystem)
 	{
 		bool bFirstItem = true;
 
 		FBoxSphereBounds Extents;
-		for (UMeshComponent* SelectedComponent : SharedMeshToolData->GetSelectedMeshComponents())
+		for (UMeshComponent* SelectedComponent : MeshPaintingSubsystem->GetSelectedMeshComponents())
 		{
 			if (bFirstItem)
 			{
@@ -260,30 +247,32 @@ bool UMeshTexturePaintingTool::Paint(const TArrayView<TPair<FVector, FVector>>&
 
 void UMeshTexturePaintingTool::CacheSelectionData()
 {
-	if (SharedMeshToolData.IsValid())
+	UMeshPaintingSubsystem* MeshPaintingSubsystem = GEngine->GetEngineSubsystem<UMeshPaintingSubsystem>();
+	if (MeshPaintingSubsystem)
 	{
-		SharedMeshToolData->ClearPaintableMeshComponents();
+		MeshPaintingSubsystem->ClearPaintableMeshComponents();
 
 		//Determine LOD level to use for painting(can only paint on LODs in vertex mode)
 		const int32 PaintLODIndex = 0;
 		//Determine UV channel to use while painting textures
 		const int32 UVChannel = 0;
 
-		SharedMeshToolData->CacheSelectionData(PaintLODIndex, UVChannel);
+		MeshPaintingSubsystem->CacheSelectionData(PaintLODIndex, UVChannel);
 	}
 }
 
 void UMeshTexturePaintingTool::CacheTexturePaintData()
 {
-	if (SharedMeshToolData.IsValid())
+	UMeshPaintingSubsystem* MeshPaintingSubsystem = GEngine->GetEngineSubsystem<UMeshPaintingSubsystem>();
+	if (MeshPaintingSubsystem)
 	{
-		TArray<UMeshComponent*> PaintableComponents = SharedMeshToolData->GetPaintableMeshComponents();
+		TArray<UMeshComponent*> PaintableComponents = MeshPaintingSubsystem->GetPaintableMeshComponents();
 		
 		PaintableTextures.Empty();
 		if (PaintableComponents.Num() == 1 && PaintableComponents[0])
 		{
 			const UMeshComponent* Component = PaintableComponents[0];
-			TSharedPtr<IMeshPaintComponentAdapter> Adapter = SharedMeshToolData->GetAdapterForComponent(Component);
+			TSharedPtr<IMeshPaintComponentAdapter> Adapter = MeshPaintingSubsystem->GetAdapterForComponent(Component);
 			UTexturePaintToolset::RetrieveTexturesForComponent(Component, Adapter.Get(), PaintableTextures);
 		}
 
@@ -312,7 +301,7 @@ bool UMeshTexturePaintingTool::PaintInternal(const TArrayView<TPair<FVector, FVe
 	const float InStrengthScale = PaintStrength;;
 
 	bool bPaintApplied = false;
-
+	UMeshPaintingSubsystem* MeshPaintingSubsystem = GEngine->GetEngineSubsystem<UMeshPaintingSubsystem>();
 	// Fire out a ray to see if there is a *selected* component under the mouse cursor that can be painted.
 	for (int32 i = 0; i < Rays.Num(); ++i)
 	{
@@ -323,11 +312,11 @@ bool UMeshTexturePaintingTool::PaintInternal(const TArrayView<TPair<FVector, FVe
 		const FVector TraceStart(RayOrigin);
 		const FVector TraceEnd(RayOrigin + RayDirection * HALF_WORLD_MAX);
 
-		if (SharedMeshToolData.IsValid())
+		if (MeshPaintingSubsystem)
 		{
-			for (UMeshComponent* MeshComponent : SharedMeshToolData->GetPaintableMeshComponents())
+			for (UMeshComponent* MeshComponent : MeshPaintingSubsystem->GetPaintableMeshComponents())
 			{
-				TSharedPtr<IMeshPaintComponentAdapter> MeshAdapter = SharedMeshToolData->GetAdapterForComponent(MeshComponent);
+				TSharedPtr<IMeshPaintComponentAdapter> MeshAdapter = MeshPaintingSubsystem->GetAdapterForComponent(MeshComponent);
 
 				// Ray trace
 				FHitResult TraceHitResult(1.0f);
@@ -412,8 +401,7 @@ bool UMeshTexturePaintingTool::PaintInternal(const TArrayView<TPair<FVector, FVe
 		{
 			UMeshComponent* HoveredComponent = Entry.Key;
 			TArray<int32>& PaintRayResultIds = Entry.Value;
-
-			IMeshPaintComponentAdapter* MeshAdapter = SharedMeshToolData.IsValid() ? SharedMeshToolData->GetAdapterForComponent(HoveredComponent).Get() : nullptr;
+			IMeshPaintComponentAdapter* MeshAdapter = MeshPaintingSubsystem ? MeshPaintingSubsystem->GetAdapterForComponent(HoveredComponent).Get() : nullptr;
 			if (!ensure(MeshAdapter))
 			{
 				continue;
@@ -442,7 +430,7 @@ bool UMeshTexturePaintingTool::PaintInternal(const TArrayView<TPair<FVector, FVe
 						const FVector& BestTraceResultLocation = PaintRayResults[PaintRayResultId].BestTraceResult.Location;
 						FViewCameraState CameraState;
 						GetToolManager()->GetContextQueriesAPI()->GetCurrentViewState(CameraState);
-						bPaintApplied |= GEngine->GetEngineSubsystem<UMeshPaintingSubsystem>()->ApplyPerTrianglePaintAction(MeshAdapter, CameraState.Position, BestTraceResultLocation, BrushProperties, FPerTrianglePaintAction::CreateUObject(this, &UMeshTexturePaintingTool::GatherTextureTriangles, &TrianglePaintInfoArray, &MaterialSections, TextureProperties->UVChannel), TextureProperties->bOnlyFrontFacingTriangles);
+						bPaintApplied |= MeshPaintingSubsystem->ApplyPerTrianglePaintAction(MeshAdapter, CameraState.Position, BestTraceResultLocation, BrushProperties, FPerTrianglePaintAction::CreateUObject(this, &UMeshTexturePaintingTool::GatherTextureTriangles, &TrianglePaintInfoArray, &MaterialSections, TextureProperties->UVChannel), TextureProperties->bOnlyFrontFacingTriangles);
 						break;
 					}
 
@@ -531,9 +519,9 @@ void UMeshTexturePaintingTool::OnEndDrag(const FRay& Ray)
 bool UMeshTexturePaintingTool::HitTest(const FRay& Ray, FHitResult& OutHit)
 {
 	bool bUsed = false;
-	if (SharedMeshToolData.IsValid())
+	if (UMeshPaintingSubsystem* MeshPaintingSubsystem = GEngine->GetEngineSubsystem<UMeshPaintingSubsystem>())
 	{
-		SharedMeshToolData->FindHitResult(Ray, OutHit);
+		MeshPaintingSubsystem->FindHitResult(Ray, OutHit);
 		LastBestHitResult = OutHit;
 		bUsed = OutHit.bBlockingHit;
 	}
@@ -1307,10 +1295,4 @@ void UMeshTexturePaintingTool::PaintTextureChanged(const FAssetData& AssetData)
 		}
 	}
 }
-
-void UMeshTexturePaintingTool::SetMeshToolData(TWeakObjectPtr<UMeshToolManager> InToolData)
-{
-	SharedMeshToolData = InToolData;
-}
-
 #undef LOCTEXT_NAMESPACE
diff --git a/Engine/Plugins/MeshPainting/Source/MeshPaintingToolset/Private/MeshVertexPaintingTool.cpp b/Engine/Plugins/MeshPainting/Source/MeshPaintingToolset/Private/MeshVertexPaintingTool.cpp
--- a/Engine/Plugins/MeshPainting/Source/MeshPaintingToolset/Private/MeshVertexPaintingTool.cpp
+++ b/Engine/Plugins/MeshPainting/Source/MeshPaintingToolset/Private/MeshVertexPaintingTool.cpp
@@ -29,7 +29,6 @@ bool UMeshColorPaintingToolBuilder::CanBuildTool(const FToolBuilderState& SceneS
 UInteractiveTool* UMeshColorPaintingToolBuilder::BuildTool(const FToolBuilderState& SceneState) const
 {
 	UMeshColorPaintingTool* NewTool = NewObject<UMeshColorPaintingTool>(SceneState.ToolManager);
-	NewTool->SetMeshToolData(SharedMeshToolData);
 	return NewTool;
 }
 
@@ -41,7 +40,6 @@ bool UMeshWeightPaintingToolBuilder::CanBuildTool(const FToolBuilderState& Scene
 UInteractiveTool* UMeshWeightPaintingToolBuilder::BuildTool(const FToolBuilderState& SceneState) const
 {
 	UMeshWeightPaintingTool* NewTool = NewObject<UMeshWeightPaintingTool>(SceneState.ToolManager);
-	NewTool->SetMeshToolData(SharedMeshToolData);
 	return NewTool;
 }
 
@@ -69,11 +67,6 @@ bool UMeshVertexPaintingTool::IsMeshAdapterSupported(TSharedPtr<IMeshPaintCompon
 	return MeshAdapter.IsValid() ? MeshAdapter->SupportsVertexPaint() : false;
 }
 
-void UMeshVertexPaintingTool::SetMeshToolData(TWeakObjectPtr<UMeshToolManager> InMeshToolData)
-{
-	SharedMeshToolData = InMeshToolData;
-}
-
 void UMeshVertexPaintingTool::Setup()
 {
 	Super::Setup();
@@ -92,7 +85,6 @@ void UMeshVertexPaintingTool::Setup()
 	// Set up selection mechanic to find and select edges
 	SelectionMechanic = NewObject<UMeshPaintSelectionMechanic>(this);
 	SelectionMechanic->Setup(this);
-	SelectionMechanic->SetMeshToolData(SharedMeshToolData);
 }
 
 
@@ -100,9 +92,10 @@ void UMeshVertexPaintingTool::Shutdown(EToolShutdownType ShutdownType)
 {
 	FinishPainting();
 	BrushProperties->SaveProperties(this);
-	if (SharedMeshToolData.IsValid())
+	UMeshPaintingSubsystem* MeshPaintingSubsystem = GEngine->GetEngineSubsystem<UMeshPaintingSubsystem>();
+	if (MeshPaintingSubsystem)
 	{
-		SharedMeshToolData->Refresh();
+		MeshPaintingSubsystem->Refresh();
 	}
 	Super::Shutdown(ShutdownType);
 }
@@ -112,7 +105,8 @@ void UMeshVertexPaintingTool::Render(IToolsContextRenderAPI* RenderAPI)
 	Super::Render(RenderAPI);
 	FToolDataVisualizer Draw;
 	Draw.BeginFrame(RenderAPI);
-	if (SharedMeshToolData.IsValid() && LastBestHitResult.Component != nullptr && SharedMeshToolData->GetPaintableMeshComponents().Num() > 0)
+	UMeshPaintingSubsystem* MeshPaintingSubsystem = GEngine->GetEngineSubsystem<UMeshPaintingSubsystem>();
+	if (MeshPaintingSubsystem && LastBestHitResult.Component != nullptr && MeshPaintingSubsystem->GetPaintableMeshComponents().Num() > 0)
 	{
 		BrushStampIndicator->bDrawIndicatorLines = true;
 		static float WidgetLineThickness = 1.0f;
@@ -127,9 +121,9 @@ void UMeshVertexPaintingTool::Render(IToolsContextRenderAPI* RenderAPI)
 		const FVector NormalLineEnd(LastBestHitResult.Location + LastBestHitResult.Normal * NormalLineSize);
 		Draw.DrawLine(FVector(LastBestHitResult.Location), NormalLineEnd, NormalLineColor, WidgetLineThickness);
 
-		for (UMeshComponent* CurrentComponent : SharedMeshToolData->GetPaintableMeshComponents())
+		for (UMeshComponent* CurrentComponent : MeshPaintingSubsystem->GetPaintableMeshComponents())
 		{
-			TSharedPtr<IMeshPaintComponentAdapter> MeshAdapter = SharedMeshToolData->GetAdapterForComponent(Cast<UMeshComponent>(CurrentComponent));
+			TSharedPtr<IMeshPaintComponentAdapter> MeshAdapter = MeshPaintingSubsystem->GetAdapterForComponent(Cast<UMeshComponent>(CurrentComponent));
 
 			if (MeshAdapter->IsValid() && MeshAdapter->SupportsVertexPaint())
 			{
@@ -169,9 +163,9 @@ void UMeshVertexPaintingTool::Render(IToolsContextRenderAPI* RenderAPI)
 
 void UMeshVertexPaintingTool::OnTick(float DeltaTime)
 {
-	if (SharedMeshToolData.IsValid())
+	if (UMeshPaintingSubsystem* MeshPaintingSubsystem = GEngine->GetEngineSubsystem<UMeshPaintingSubsystem>())
 	{
-		if (SharedMeshToolData->bNeedsRecache)
+		if (MeshPaintingSubsystem->bNeedsRecache)
 		{
 			CacheSelectionData();
 			bDoRestoreRenTargets = true;
@@ -202,12 +196,12 @@ double UMeshVertexPaintingTool::EstimateMaximumTargetDimension()
 {
 	bool bFoundComponentToUse = false;
 	FBoxSphereBounds Bounds = FBoxSphereBounds(0.0);
-	if (SharedMeshToolData.IsValid())
+	if (UMeshPaintingSubsystem* MeshPaintingSubsystem = GEngine->GetEngineSubsystem<UMeshPaintingSubsystem>())
 	{
 		bool bFirstItem = true;
 
 		FBoxSphereBounds Extents;
-		for (UMeshComponent* SelectedComponent : SharedMeshToolData->GetSelectedMeshComponents())
+		for (UMeshComponent* SelectedComponent : MeshPaintingSubsystem->GetSelectedMeshComponents())
 		{
 			if (bFirstItem)
 			{
@@ -265,7 +259,7 @@ bool UMeshVertexPaintingTool::PaintInternal(const TArrayView<TPair<FVector, FVec
 	PaintRayResults.AddDefaulted(Rays.Num());
 	LastBestHitResult.Reset();
 	TMap<UMeshComponent*, TArray<int32>> HoveredComponents;
-
+	UMeshPaintingSubsystem* MeshPaintingSubsystem = GEngine->GetEngineSubsystem<UMeshPaintingSubsystem>();
 	const float BrushRadius = BrushProperties->BrushRadius;
 	const bool bIsPainting = PaintAction == EMeshPaintModeAction::Paint;
 	const float InStrengthScale = PaintStrength;
@@ -281,9 +275,9 @@ bool UMeshVertexPaintingTool::PaintInternal(const TArrayView<TPair<FVector, FVec
 		const FVector TraceStart(RayOrigin);
 		const FVector TraceEnd(RayOrigin + RayDirection * HALF_WORLD_MAX);
 
-		for (UMeshComponent* MeshComponent : SharedMeshToolData->GetPaintableMeshComponents())
+		for (UMeshComponent* MeshComponent : MeshPaintingSubsystem->GetPaintableMeshComponents())
 		{
-			TSharedPtr<IMeshPaintComponentAdapter> MeshAdapter = SharedMeshToolData->GetAdapterForComponent(MeshComponent);
+			TSharedPtr<IMeshPaintComponentAdapter> MeshAdapter = MeshPaintingSubsystem->GetAdapterForComponent(MeshComponent);
 
 			// Ray trace
 			FHitResult TraceHitResult(1.0f);
@@ -305,11 +299,11 @@ bool UMeshVertexPaintingTool::PaintInternal(const TArrayView<TPair<FVector, FVec
 			FBox BrushBounds = FBox::BuildAABB(BestTraceResult.Location, FVector(BrushRadius * 1.25f, BrushRadius * 1.25f, BrushRadius * 1.25f));
 
 			// Vertex paint mode, so we want all valid components overlapping the brush hit location
-			for (auto TestComponent : SharedMeshToolData->GetPaintableMeshComponents())
+			for (auto TestComponent : MeshPaintingSubsystem->GetPaintableMeshComponents())
 			{
 				const FBox ComponentBounds = TestComponent->Bounds.GetBox();
 
-				if (SharedMeshToolData->GetComponentToAdapterMap().Contains(TestComponent) && ComponentBounds.Intersect(BrushBounds))
+				if (MeshPaintingSubsystem->GetComponentToAdapterMap().Contains(TestComponent->GetFName()) && ComponentBounds.Intersect(BrushBounds))
 				{
 					// OK, this mesh potentially overlaps the brush!
 					HoveredComponents.FindOrAdd(TestComponent).Add(i);
@@ -371,7 +365,7 @@ bool UMeshVertexPaintingTool::PaintInternal(const TArrayView<TPair<FVector, FVec
 			UMeshComponent* HoveredComponent = Entry.Key;
 			TArray<int32>& PaintRayResultIds = Entry.Value;
 
-			IMeshPaintComponentAdapter* MeshAdapter = SharedMeshToolData->GetAdapterForComponent(HoveredComponent).Get();
+			IMeshPaintComponentAdapter* MeshAdapter = MeshPaintingSubsystem->GetAdapterForComponent(HoveredComponent).Get();
 			if (!ensure(MeshAdapter))
 			{
 				continue;
@@ -394,7 +388,7 @@ bool UMeshVertexPaintingTool::PaintInternal(const TArrayView<TPair<FVector, FVec
 				{
 					InfluencedVertices.Reset();
 					Args.HitResult = PaintRayResults[PaintRayResultId].BestTraceResult;
-					bPaintApplied |= GEngine->GetEngineSubsystem<UMeshPaintingSubsystem>()->GetPerVertexPaintInfluencedVertices(Args, InfluencedVertices);
+					bPaintApplied |= MeshPaintingSubsystem->GetPerVertexPaintInfluencedVertices(Args, InfluencedVertices);
 
 					if (InfluencedVertices.Num() == 0)
 					{
@@ -522,9 +516,10 @@ void UMeshVertexPaintingTool::OnEndDrag(const FRay& Ray)
 bool UMeshVertexPaintingTool::HitTest(const FRay& Ray, FHitResult& OutHit)
 {
 	bool bUsed = false;
-	if (SharedMeshToolData.IsValid())
+	UMeshPaintingSubsystem* MeshPaintingSubsystem = GEngine->GetEngineSubsystem<UMeshPaintingSubsystem>();
+	if (MeshPaintingSubsystem)
 	{
-		SharedMeshToolData->FindHitResult(Ray, OutHit);
+		MeshPaintingSubsystem->FindHitResult(Ray, OutHit);
 		LastBestHitResult = OutHit;
 		bUsed = OutHit.bBlockingHit;
 	}
@@ -577,9 +572,9 @@ void UMeshColorPaintingTool::Shutdown(EToolShutdownType ShutdownType)
 void UMeshColorPaintingTool::CacheSelectionData()
 {
 	Super::CacheSelectionData();
-	if (SharedMeshToolData.IsValid())
+	if (UMeshPaintingSubsystem* MeshPaintingSubsystem = GEngine->GetEngineSubsystem<UMeshPaintingSubsystem>())
 	{
-		SharedMeshToolData->ClearPaintableMeshComponents();
+		MeshPaintingSubsystem->ClearPaintableMeshComponents();
 		// Update(cached) Paint LOD level if necessary
 		ColorProperties->LODIndex = FMath::Min<int32>(ColorProperties->LODIndex, GetMaxLODIndexToPaint());
 		CachedLODIndex = ColorProperties->LODIndex;
@@ -589,7 +584,7 @@ void UMeshColorPaintingTool::CacheSelectionData()
 		//Determine UV channel to use while painting textures
 		const int32 UVChannel = 0;
 
-		SharedMeshToolData->CacheSelectionData(PaintLODIndex, UVChannel);
+		MeshPaintingSubsystem->CacheSelectionData(PaintLODIndex, UVChannel);
 	}
 }
 
@@ -606,15 +601,14 @@ int32 UMeshColorPaintingTool::GetMaxLODIndexToPaint() const
 {
 	//The maximum LOD we can paint is decide by the lowest number of LOD in the selection
 	int32 LODMin = TNumericLimits<int32>::Max();
-
-	if (SharedMeshToolData.IsValid())
+	if (UMeshPaintingSubsystem* MeshPaintingSubsystem = GEngine->GetEngineSubsystem<UMeshPaintingSubsystem>())
 	{
-		TArray<UMeshComponent*> SelectedComponents = SharedMeshToolData->GetSelectedMeshComponents();
+		TArray<UMeshComponent*> SelectedComponents = MeshPaintingSubsystem->GetSelectedMeshComponents();
 
 		for (UMeshComponent* MeshComponent : SelectedComponents)
 		{
 			int32 NumMeshLODs = 0;
-			if (GEngine->GetEngineSubsystem<UMeshPaintingSubsystem>()->TryGetNumberOfLODs(MeshComponent, NumMeshLODs))
+			if (MeshPaintingSubsystem->TryGetNumberOfLODs(MeshComponent, NumMeshLODs))
 			{
 				ensure(NumMeshLODs > 0);
 				LODMin = FMath::Min(LODMin, NumMeshLODs - 1);
@@ -645,9 +639,9 @@ void UMeshColorPaintingTool::LODPaintStateChanged(const bool bLODPaintingEnabled
 
 	ApplyForcedLODIndex(bLODPaintingEnabled ? CachedLODIndex : -1);
 
-	if (SharedMeshToolData.IsValid())
+	if (UMeshPaintingSubsystem* MeshPaintingSubsystem = GEngine->GetEngineSubsystem<UMeshPaintingSubsystem>())
 	{
-		TArray<UMeshComponent*> PaintableComponents = SharedMeshToolData->GetPaintableMeshComponents();
+		TArray<UMeshComponent*> PaintableComponents = MeshPaintingSubsystem->GetPaintableMeshComponents();
 
 		TUniquePtr< FComponentReregisterContext > ComponentReregisterContext;
 		//Make sure all static mesh render is dirty since we change the force LOD
@@ -659,21 +653,21 @@ void UMeshColorPaintingTool::LODPaintStateChanged(const bool bLODPaintingEnabled
 			}
 		}
 
-		SharedMeshToolData->Refresh();
+		MeshPaintingSubsystem->Refresh();
 	}
 }
 
 void UMeshColorPaintingTool::ApplyForcedLODIndex(int32 ForcedLODIndex)
 {
-	if (SharedMeshToolData.IsValid())
+	if (UMeshPaintingSubsystem* MeshPaintingSubsystem = GEngine->GetEngineSubsystem<UMeshPaintingSubsystem>())
 	{
-		TArray<UMeshComponent*> PaintableComponents = SharedMeshToolData->GetPaintableMeshComponents();
+		TArray<UMeshComponent*> PaintableComponents = MeshPaintingSubsystem->GetPaintableMeshComponents();
 
 		for (UMeshComponent* SelectedComponent : PaintableComponents)
 		{
 			if (SelectedComponent)
 			{
-				GEngine->GetEngineSubsystem<UMeshPaintingSubsystem>()->ForceRenderMeshLOD(SelectedComponent, ForcedLODIndex);
+				MeshPaintingSubsystem->ForceRenderMeshLOD(SelectedComponent, ForcedLODIndex);
 			}
 		}
 	}
@@ -690,9 +684,9 @@ void UMeshColorPaintingTool::PaintLODChanged()
 
 		TUniquePtr< FComponentReregisterContext > ComponentReregisterContext;
 		//Make sure all static mesh render is dirty since we change the force LOD
-		if (SharedMeshToolData.IsValid())
+		if (UMeshPaintingSubsystem* MeshPaintingSubsystem = GEngine->GetEngineSubsystem<UMeshPaintingSubsystem>())
 		{
-			TArray<UMeshComponent*> PaintableComponents = SharedMeshToolData->GetPaintableMeshComponents();
+			TArray<UMeshComponent*> PaintableComponents = MeshPaintingSubsystem->GetPaintableMeshComponents();
 
 			for (UMeshComponent* SelectedComponent : PaintableComponents)
 			{
@@ -702,7 +696,7 @@ void UMeshColorPaintingTool::PaintLODChanged()
 				}
 			}
 
-			SharedMeshToolData->Refresh();
+			MeshPaintingSubsystem->Refresh();
 		}
 	}
 }
@@ -746,14 +740,14 @@ void UMeshWeightPaintingTool::Setup()
 void UMeshWeightPaintingTool::CacheSelectionData()
 {
 	Super::CacheSelectionData();
-	if (SharedMeshToolData.IsValid())
+	if (UMeshPaintingSubsystem* MeshPaintingSubsystem = GEngine->GetEngineSubsystem<UMeshPaintingSubsystem>())
 	{
-		SharedMeshToolData->ClearPaintableMeshComponents();
+		MeshPaintingSubsystem->ClearPaintableMeshComponents();
 		//Determine LOD level to use for painting(can only paint on LODs in vertex mode)
 		const int32 PaintLODIndex = 0;
 		//Determine UV channel to use while painting textures
 		const int32 UVChannel = 0;
-		SharedMeshToolData->CacheSelectionData(PaintLODIndex, UVChannel);
+		MeshPaintingSubsystem->CacheSelectionData(PaintLODIndex, UVChannel);
 	}
 }
 
diff --git a/Engine/Plugins/MeshPainting/Source/MeshPaintingToolset/Public/MeshPaintHelpers.h b/Engine/Plugins/MeshPainting/Source/MeshPaintingToolset/Public/MeshPaintHelpers.h
--- a/Engine/Plugins/MeshPainting/Source/MeshPaintingToolset/Public/MeshPaintHelpers.h
+++ b/Engine/Plugins/MeshPainting/Source/MeshPaintingToolset/Public/MeshPaintHelpers.h
@@ -128,6 +128,8 @@ class MESHPAINTINGTOOLSET_API UMeshPaintingSubsystem : public UEngineSubsystem
 
 	GENERATED_BODY()
 public:
+	UMeshPaintingSubsystem();
+
 	bool HasPaintableMesh(UActorComponent* Component);
 	/** Removes vertex colors associated with the object */
 	void RemoveInstanceVertexColors(UObject* Obj);
@@ -228,35 +230,19 @@ public:
 
 	/** Helper function to retrieve vertex color from a UTexture given a UVCoordinate */
 	FColor PickVertexColorFromTextureData(const uint8* MipData, const FVector2D& UVCoordinate, const UTexture2D* Texture, const FColor ColorMask);	
-};
 
-template<typename T>
-void UMeshPaintingSubsystem::ApplyBrushToVertex(const FVector& VertexPosition, const FMatrix& InverseBrushMatrix, const float BrushRadius, const float BrushFalloffAmount, const float BrushStrength, const T& PaintValue, T& InOutValue)
-{
-	const FVector BrushSpacePosition = InverseBrushMatrix.TransformPosition(VertexPosition);
-	const FVector2D BrushSpacePosition2D(BrushSpacePosition.X, BrushSpacePosition.Y);
-		
-	float InfluencedValue = 0.0f;
-	if (IsPointInfluencedByBrush(BrushSpacePosition2D, BrushRadius * BrushRadius, InfluencedValue))
+	void SetSelectionHasMaterialValidForTexturePaint(const bool InValidity)
 	{
-		float InnerBrushRadius = BrushFalloffAmount * BrushRadius;
-		float PaintStrength = GEngine->GetEngineSubsystem<UMeshPaintingSubsystem>()->ComputePaintMultiplier(BrushSpacePosition2D.SizeSquared(), BrushStrength, InnerBrushRadius, BrushRadius - InnerBrushRadius, 1.0f, 1.0f, 1.0f);
-
-		const T OldValue = InOutValue;
-		InOutValue = FMath::LerpStable(OldValue, PaintValue, PaintStrength);
-	}	
-};
-
-UCLASS()
-class MESHPAINTINGTOOLSET_API UMeshToolManager : public UObject
-{
-	GENERATED_BODY()
+		bSelectionHasMaterialValidForTexturePaint = InValidity;
+	}
 
-public:
-	UMeshToolManager();
+	bool SelectionHasMaterialValidForTexturePaint()
+	{
+		return bSelectionHasMaterialValidForTexturePaint;
+	}
 
 	/** Map of geometry adapters for each selected mesh component */
-	TMap<UMeshComponent*, TSharedPtr<IMeshPaintComponentAdapter>> GetComponentToAdapterMap() const;
+	TMap<FName, TSharedPtr<IMeshPaintComponentAdapter>> GetComponentToAdapterMap() const;
 	TSharedPtr<IMeshPaintComponentAdapter> GetAdapterForComponent(const UMeshComponent* InComponent);
 	void AddToComponentToAdapterMap(UMeshComponent* InComponent, TSharedPtr<IMeshPaintComponentAdapter> InAdapter);
 
@@ -280,17 +266,36 @@ public:
 public:
 	bool bNeedsRecache;
 
+protected:
+	bool bSelectionHasMaterialValidForTexturePaint;
 
 private:
 	void CleanUp();
 
 private:
 	/** Map of geometry adapters for each selected mesh component */
-	TMap<UMeshComponent*, TSharedPtr<IMeshPaintComponentAdapter>> ComponentToAdapterMap;
+	TMap<FName, TSharedPtr<IMeshPaintComponentAdapter>> ComponentToAdapterMap;
 	TArray<UMeshComponent*> SelectedMeshComponents;
 	/** Mesh components within the current selection which are eligible for painting */
 	TArray<UMeshComponent*> PaintableComponents;
 	/** Contains copied vertex color data */
 	TArray<FPerComponentVertexColorData> CopiedColorsByComponent;
 	bool bSelectionContainsPerLODColors;
-};
\ No newline at end of file
+};
+
+template<typename T>
+void UMeshPaintingSubsystem::ApplyBrushToVertex(const FVector& VertexPosition, const FMatrix& InverseBrushMatrix, const float BrushRadius, const float BrushFalloffAmount, const float BrushStrength, const T& PaintValue, T& InOutValue)
+{
+	const FVector BrushSpacePosition = InverseBrushMatrix.TransformPosition(VertexPosition);
+	const FVector2D BrushSpacePosition2D(BrushSpacePosition.X, BrushSpacePosition.Y);
+		
+	float InfluencedValue = 0.0f;
+	if (IsPointInfluencedByBrush(BrushSpacePosition2D, BrushRadius * BrushRadius, InfluencedValue))
+	{
+		float InnerBrushRadius = BrushFalloffAmount * BrushRadius;
+		float PaintStrength = GEngine->GetEngineSubsystem<UMeshPaintingSubsystem>()->ComputePaintMultiplier(BrushSpacePosition2D.SizeSquared(), BrushStrength, InnerBrushRadius, BrushRadius - InnerBrushRadius, 1.0f, 1.0f, 1.0f);
+
+		const T OldValue = InOutValue;
+		InOutValue = FMath::LerpStable(OldValue, PaintValue, PaintStrength);
+	}	
+};
diff --git a/Engine/Plugins/MeshPainting/Source/MeshPaintingToolset/Public/MeshPaintInteractions.h b/Engine/Plugins/MeshPainting/Source/MeshPaintingToolset/Public/MeshPaintInteractions.h
--- a/Engine/Plugins/MeshPainting/Source/MeshPaintingToolset/Public/MeshPaintInteractions.h
+++ b/Engine/Plugins/MeshPainting/Source/MeshPaintingToolset/Public/MeshPaintInteractions.h
@@ -39,7 +39,6 @@ public:
 	{
 		bAddToSelectionSet = bInNewSelectionType;
 	}
-	void SetMeshToolData(TWeakObjectPtr<UMeshToolManager> InMeshToolData);
 protected:
 	bool FindClickedComponentsAndCacheAdapters(const FInputDeviceRay& ClickPos);
 
@@ -48,7 +47,5 @@ protected:
 	TArray<UMeshComponent*> CachedClickedComponents;
 	UPROPERTY(Transient)
 	TArray<AActor*> CachedClickedActors;
-
-	TWeakObjectPtr<UMeshToolManager> SharedMeshToolData;
 	bool bAddToSelectionSet;
 };
diff --git a/Engine/Plugins/MeshPainting/Source/MeshPaintingToolset/Public/MeshSelect.h b/Engine/Plugins/MeshPainting/Source/MeshPaintingToolset/Public/MeshSelect.h
--- a/Engine/Plugins/MeshPainting/Source/MeshPaintingToolset/Public/MeshSelect.h
+++ b/Engine/Plugins/MeshPainting/Source/MeshPaintingToolset/Public/MeshSelect.h
@@ -11,7 +11,6 @@ class IMeshPaintComponentAdapter;
 class AActor;
 class UMeshComponent;
 class UMeshPaintSelectionMechanic;
-class UMeshToolManager;
 
 /**
  * Builder for UVertexAdapterClickTool
@@ -24,8 +23,6 @@ class MESHPAINTINGTOOLSET_API UVertexAdapterClickToolBuilder : public USingleCli
 public:
 	virtual bool CanBuildTool(const FToolBuilderState& SceneState) const override;
 	virtual UInteractiveTool* BuildTool(const FToolBuilderState& SceneState) const override;
-	
-	TWeakObjectPtr<UMeshToolManager> SharedMeshToolData;
 };
 
 
@@ -40,8 +37,6 @@ class MESHPAINTINGTOOLSET_API UTextureAdapterClickToolBuilder : public USingleCl
 public:
 	virtual bool CanBuildTool(const FToolBuilderState& SceneState) const override;
 	virtual UInteractiveTool* BuildTool(const FToolBuilderState& SceneState) const override;
-
-	TWeakObjectPtr<UMeshToolManager> SharedMeshToolData;
 };
 
 
@@ -80,8 +75,6 @@ public:
 		return true;
 	}
 
-	void SetMeshToolData(TWeakObjectPtr<UMeshToolManager> InMeshToolData);
-
 protected:
 	// flags used to identify modifier keys/buttons
 	static const int AdditiveSelectionModifier = 1;
@@ -90,7 +83,6 @@ protected:
 	UPROPERTY(Transient)
 	UMeshPaintSelectionMechanic* SelectionMechanic;
 
-	TWeakObjectPtr<UMeshToolManager> SharedMeshToolData;
 };
 
 UCLASS()
diff --git a/Engine/Plugins/MeshPainting/Source/MeshPaintingToolset/Public/MeshTexturePaintingTool.h b/Engine/Plugins/MeshPainting/Source/MeshPaintingToolset/Public/MeshTexturePaintingTool.h
--- a/Engine/Plugins/MeshPainting/Source/MeshPaintingToolset/Public/MeshTexturePaintingTool.h
+++ b/Engine/Plugins/MeshPainting/Source/MeshPaintingToolset/Public/MeshTexturePaintingTool.h
@@ -136,8 +136,6 @@ public:
 	bool ShouldFilterTextureAsset(const FAssetData& AssetData) const;
 	void PaintTextureChanged(const FAssetData& AssetData);
 
-	void SetMeshToolData(TWeakObjectPtr<UMeshToolManager> InToolData);
-
 protected:
 	virtual void SetAdditionalPaintParameters(FMeshPaintParameters& InPaintParameters) {};
 	virtual void FinishPainting();
@@ -171,8 +169,6 @@ private:
 	UPROPERTY(Transient)
 	TArray<const UTexture*> Textures;
 
-	TWeakObjectPtr<UMeshToolManager> SharedMeshToolData;
-
 	/** Flag for whether or not we are currently painting */
 	bool bArePainting;
 	bool bDoRestoreRenTargets;
diff --git a/Engine/Plugins/MeshPainting/Source/MeshPaintingToolset/Public/MeshVertexPaintingTool.h b/Engine/Plugins/MeshPainting/Source/MeshPaintingToolset/Public/MeshVertexPaintingTool.h
--- a/Engine/Plugins/MeshPainting/Source/MeshPaintingToolset/Public/MeshVertexPaintingTool.h
+++ b/Engine/Plugins/MeshPainting/Source/MeshPaintingToolset/Public/MeshVertexPaintingTool.h
@@ -63,8 +63,6 @@ class MESHPAINTINGTOOLSET_API UMeshColorPaintingToolBuilder : public UInteractiv
 public:
 	virtual bool CanBuildTool(const FToolBuilderState& SceneState) const override;
 	virtual UInteractiveTool* BuildTool(const FToolBuilderState& SceneState) const override;
-	
-	TWeakObjectPtr<UMeshToolManager> SharedMeshToolData;
 };
 
 UCLASS()
@@ -75,8 +73,6 @@ class MESHPAINTINGTOOLSET_API UMeshWeightPaintingToolBuilder : public UInteracti
 public:
 	virtual bool CanBuildTool(const FToolBuilderState& SceneState) const override;
 	virtual UInteractiveTool* BuildTool(const FToolBuilderState& SceneState) const override;
-
-	TWeakObjectPtr<UMeshToolManager> SharedMeshToolData;
 };
 
 
@@ -204,8 +200,6 @@ public:
 		return true;
 	}
 
-	void SetMeshToolData(TWeakObjectPtr<UMeshToolManager> InMeshToolData);
-
 protected:
 	virtual void SetAdditionalPaintParameters(FMeshPaintParameters& InPaintParameters) {};
 	virtual void FinishPainting();
@@ -231,8 +225,6 @@ protected:
 	UPROPERTY(Transient)
 	UMeshPaintSelectionMechanic* SelectionMechanic;
 
-	TWeakObjectPtr<UMeshToolManager> SharedMeshToolData;
-
 private:
 	bool PaintInternal(const TArrayView<TPair<FVector, FVector>>& Rays, EMeshPaintModeAction PaintAction, float PaintStrength);
 	
