Checking MeshComponent ptr against nullptr

diff --git a/Engine/Plugins/MeshPainting/Source/MeshPaintingToolset/Private/MeshPaintInteractions.cpp b/Engine/Plugins/MeshPainting/Source/MeshPaintingToolset/Private/MeshPaintInteractions.cpp
--- a/Engine/Plugins/MeshPainting/Source/MeshPaintingToolset/Private/MeshPaintInteractions.cpp
+++ b/Engine/Plugins/MeshPainting/Source/MeshPaintingToolset/Private/MeshPaintInteractions.cpp
@@ -34,7 +34,11 @@ void UMeshPaintSelectionMechanic::OnClicked(const FInputDeviceRay& ClickPos)
 	{
 		for (UMeshComponent* MeshComponent : CachedClickedComponents)
 		{
-			TSharedPtr<IMeshPaintComponentAdapter> MeshAdapter = MeshPaintingSubsystem->GetAdapterForComponent(MeshComponent);
+			TSharedPtr<IMeshPaintComponentAdapter> MeshAdapter;
+			if (MeshComponent)
+			{
+				MeshAdapter = MeshPaintingSubsystem->GetAdapterForComponent(MeshComponent);
+			}
 			IMeshPaintSelectionInterface* Interface = Cast<IMeshPaintSelectionInterface>(GetParentTool());
 			if (MeshComponent && MeshComponent->IsVisible() && MeshAdapter.IsValid() && MeshAdapter->IsValid() && Interface->IsMeshAdapterSupported(MeshAdapter))
 			{
