GeometryProcessing: move all public types in GeometryProcessing plugin into UE::Geometry:: namespace, as well as some code in MeshModelingToolset. Fix fallout in Water, Paint Mode, ChaosEditor, DataPr

diff --git a/Engine/Plugins/Experimental/Water/Source/Runtime/Private/WaterBooleanUtils.h b/Engine/Plugins/Experimental/Water/Source/Runtime/Private/WaterBooleanUtils.h
--- a/Engine/Plugins/Experimental/Water/Source/Runtime/Private/WaterBooleanUtils.h
+++ b/Engine/Plugins/Experimental/Water/Source/Runtime/Private/WaterBooleanUtils.h
@@ -32,6 +32,8 @@ namespace UE {
 
 	namespace WaterUtils {
 
+		using namespace UE::Geometry;
+
 		/** @return triangle aspect ratio transformed to be in [0,1] range */
 		double UnitAspectRatio(const FVector3d& A, const FVector3d& B, const FVector3d& C)
 		{
@@ -195,7 +197,7 @@ namespace UE {
 			// Merge the mesh edges to create a closed solid
 			double MinLen, MaxLen, AvgLen;
 			TMeshQueries<FDynamicMesh3>::EdgeLengthStats(Mesh, MinLen, MaxLen, AvgLen);
-			FMergeCoincidentMeshEdges Merge(&Mesh);
+			UE::Geometry::FMergeCoincidentMeshEdges Merge(&Mesh);
 			Merge.MergeVertexTolerance = FMathd::Max(Merge.MergeVertexTolerance, MinLen * 0.1);
 			Merge.Apply();
 
diff --git a/Engine/Plugins/MeshPainting/Source/MeshPaintingToolset/Private/BaseMeshPaintComponentAdapter.cpp b/Engine/Plugins/MeshPainting/Source/MeshPaintingToolset/Private/BaseMeshPaintComponentAdapter.cpp
--- a/Engine/Plugins/MeshPainting/Source/MeshPaintingToolset/Private/BaseMeshPaintComponentAdapter.cpp
+++ b/Engine/Plugins/MeshPainting/Source/MeshPaintingToolset/Private/BaseMeshPaintComponentAdapter.cpp
@@ -10,7 +10,7 @@ bool FBaseMeshPaintComponentAdapter::Initialize()
 	bool bInit = InitializeVertexData() && BuildOctree();
 	FIndexMeshArrayAdapterd TempAdapter(&MeshVertices, &MeshIndices);
 	Adapter = TempAdapter;
-	AABBTree = MakeUnique< TMeshAABBTree3<const FIndexMeshArrayAdapterd>>(&Adapter);
+	AABBTree = MakeUnique<UE::Geometry::TMeshAABBTree3<const FIndexMeshArrayAdapterd>>(&Adapter);
 	return bInit;
 }
 
@@ -182,13 +182,13 @@ TArray<FVector> FBaseMeshPaintComponentAdapter::SphereIntersectVertices(const fl
 	return InRangeVertices;
 }
 
-bool FBaseMeshPaintComponentAdapter::RayIntersectAdapter(FIndex3i& HitTriangle, FVector& HitPosition, const FVector Start, const FVector End) const
+bool FBaseMeshPaintComponentAdapter::RayIntersectAdapter(UE::Geometry::FIndex3i& HitTriangle, FVector& HitPosition, const FVector Start, const FVector End) const
 {
 	int32 HitTriangleID;
 	double NearestT;
-	FVector3d Direction((FVector3d)End - (FVector3d)Start);
+	UE::Geometry::FVector3d Direction((UE::Geometry::FVector3d)End - (UE::Geometry::FVector3d)Start);
 	Direction.Normalize();
-	FRay3d Ray((FVector3d)Start, Direction);
+	UE::Geometry::FRay3d Ray((UE::Geometry::FVector3d)Start, Direction);
 	bool bHit = AABBTree->FindNearestHitTriangle(Ray, NearestT, HitTriangleID);
 	if (bHit)
 	{
diff --git a/Engine/Plugins/MeshPainting/Source/MeshPaintingToolset/Private/MeshPaintSkeletalMeshAdapter.cpp b/Engine/Plugins/MeshPainting/Source/MeshPaintingToolset/Private/MeshPaintSkeletalMeshAdapter.cpp
--- a/Engine/Plugins/MeshPainting/Source/MeshPaintingToolset/Private/MeshPaintSkeletalMeshAdapter.cpp
+++ b/Engine/Plugins/MeshPainting/Source/MeshPaintingToolset/Private/MeshPaintSkeletalMeshAdapter.cpp
@@ -259,7 +259,7 @@ bool FMeshPaintSkeletalMeshComponentAdapter::LineTraceComponent(struct FHitResul
 		float MinDistance = FLT_MAX;
 		FVector Intersect;
 		FVector Normal;
-		FIndex3i FoundTriangle;
+		UE::Geometry::FIndex3i FoundTriangle;
 		FVector HitPosition;
 		if (!RayIntersectAdapter(FoundTriangle, HitPosition, LocalStart, LocalEnd))
 		{
diff --git a/Engine/Plugins/MeshPainting/Source/MeshPaintingToolset/Private/MeshPaintStaticMeshAdapter.cpp b/Engine/Plugins/MeshPainting/Source/MeshPaintingToolset/Private/MeshPaintStaticMeshAdapter.cpp
--- a/Engine/Plugins/MeshPainting/Source/MeshPaintingToolset/Private/MeshPaintStaticMeshAdapter.cpp
+++ b/Engine/Plugins/MeshPainting/Source/MeshPaintingToolset/Private/MeshPaintStaticMeshAdapter.cpp
@@ -200,7 +200,7 @@ bool FMeshPaintStaticMeshComponentAdapter::LineTraceComponent(struct FHitResult&
 		float MinDistance = FLT_MAX;
 		FVector Intersect;
 		FVector Normal;
-		FIndex3i FoundTriangle;
+		UE::Geometry::FIndex3i FoundTriangle;
 		FVector HitPosition;
 		if (!RayIntersectAdapter(FoundTriangle, HitPosition, LocalStart, LocalEnd))
 		{
@@ -212,14 +212,14 @@ bool FMeshPaintStaticMeshComponentAdapter::LineTraceComponent(struct FHitResult&
 		const FVector& P1 = MeshVertices[FoundTriangle.B];
 		const FVector& P2 = MeshVertices[FoundTriangle.C];
 
-		FTriangle3d Triangle((FVector3d)P0, (FVector3d)P1, (FVector3d)P2);
-		FVector3d TriNormal = Triangle.Normal();
+		UE::Geometry::FTriangle3d Triangle((UE::Geometry::FVector3d)P0, (UE::Geometry::FVector3d)P1, (UE::Geometry::FVector3d)P2);
+		UE::Geometry::FVector3d TriNormal = Triangle.Normal();
 
 		//check collinearity of A,B,C
 		if (TriNormal.SquaredLength() > (double)SMALL_NUMBER)
 		{
-			FRay3d LocalRay((FVector3d)LocalStart, ((FVector3d)LocalEnd - (FVector3d)LocalStart).Normalized());
-			FIntrRay3Triangle3d RayTriIntersection(LocalRay, Triangle);
+			UE::Geometry::FRay3d LocalRay((UE::Geometry::FVector3d)LocalStart, ((UE::Geometry::FVector3d)LocalEnd - (UE::Geometry::FVector3d)LocalStart).Normalized());
+			UE::Geometry::FIntrRay3Triangle3d RayTriIntersection(LocalRay, Triangle);
 			if (RayTriIntersection.Find())
 			{
 				double Distance = RayTriIntersection.RayParameter;
diff --git a/Engine/Plugins/MeshPainting/Source/MeshPaintingToolset/Public/BaseMeshPaintComponentAdapter.h b/Engine/Plugins/MeshPainting/Source/MeshPaintingToolset/Public/BaseMeshPaintComponentAdapter.h
--- a/Engine/Plugins/MeshPainting/Source/MeshPaintingToolset/Public/BaseMeshPaintComponentAdapter.h
+++ b/Engine/Plugins/MeshPainting/Source/MeshPaintingToolset/Public/BaseMeshPaintComponentAdapter.h
@@ -8,8 +8,7 @@
 #include "Spatial/MeshAABBTree3.h"
 #include "UObject/GCObject.h"
 
-struct FIndex3i;
-typedef TIndexMeshArrayAdapter<uint32, double> FIndexMeshArrayAdapterd;
+typedef UE::Geometry::TIndexMeshArrayAdapter<uint32, double> FIndexMeshArrayAdapterd;
 
 /** Base mesh paint geometry adapter, handles basic sphere intersection using a Octree */
 class MESHPAINTINGTOOLSET_API FBaseMeshPaintComponentAdapter : public IMeshPaintComponentAdapter, public FGCObject, public TSharedFromThis<FBaseMeshPaintComponentAdapter>
@@ -24,7 +23,7 @@ public:
 	virtual void GetInfluencedVertexIndices(const float ComponentSpaceSquaredBrushRadius, const FVector& ComponentSpaceBrushPosition, const FVector& ComponentSpaceCameraPosition, const bool bOnlyFrontFacing, TSet<int32> &InfluencedVertices) const override;
 	virtual void GetInfluencedVertexData(const float ComponentSpaceSquaredBrushRadius, const FVector& ComponentSpaceBrushPosition, const FVector& ComponentSpaceCameraPosition, const bool bOnlyFrontFacing, TArray<TPair<int32, FVector>>& OutData) const override;
 	virtual TArray<FVector> SphereIntersectVertices(const float ComponentSpaceSquaredBrushRadius, const FVector& ComponentSpaceBrushPosition, const FVector& ComponentSpaceCameraPosition, const bool bOnlyFrontFacing) const override;
-	virtual bool RayIntersectAdapter(FIndex3i& HitTriangle, FVector& HitPosition, const FVector Start, const FVector End) const override;
+	virtual bool RayIntersectAdapter(UE::Geometry::FIndex3i& HitTriangle, FVector& HitPosition, const FVector Start, const FVector End) const override;
 	/** End IMeshPaintGeometryAdapter Overrides */
 
 	virtual bool InitializeVertexData() = 0;
@@ -37,5 +36,5 @@ protected:
 	/** Octree used for reducing the cost of sphere intersecting with triangles / vertices */
 	TUniquePtr<FMeshPaintTriangleOctree> MeshTriOctree;
 	FIndexMeshArrayAdapterd Adapter;
-	TUniquePtr<TMeshAABBTree3<const FIndexMeshArrayAdapterd>> AABBTree;
+	TUniquePtr<UE::Geometry::TMeshAABBTree3<const FIndexMeshArrayAdapterd>> AABBTree;
 };
diff --git a/Engine/Plugins/MeshPainting/Source/MeshPaintingToolset/Public/IMeshPaintComponentAdapter.h b/Engine/Plugins/MeshPainting/Source/MeshPaintingToolset/Public/IMeshPaintComponentAdapter.h
--- a/Engine/Plugins/MeshPainting/Source/MeshPaintingToolset/Public/IMeshPaintComponentAdapter.h
+++ b/Engine/Plugins/MeshPainting/Source/MeshPaintingToolset/Public/IMeshPaintComponentAdapter.h
@@ -7,7 +7,13 @@
 class FReferenceCollector;
 class UMeshComponent;
 class UTexture;
-struct FIndex3i;
+namespace UE
+{
+	namespace Geometry
+	{
+		struct FIndex3i;
+	}
+}
 
 /**
  * Interface for a class to provide mesh painting support for a subclass of UMeshComponent
@@ -94,6 +100,6 @@ public:
 	virtual TArray<FVector> SphereIntersectVertices(const float ComponentSpaceSquaredBrushRadius, const FVector& ComponentSpaceBrushPosition, const FVector& ComponentSpaceCameraPosition, const bool bOnlyFrontFacing) const = 0;
 
 	/** Returns the hit result which intersect with the given ray */
-	virtual bool RayIntersectAdapter(FIndex3i& HitTriangle, FVector& HitPosition, const FVector Start, const FVector End) const = 0;
+	virtual bool RayIntersectAdapter(UE::Geometry::FIndex3i& HitTriangle, FVector& HitPosition, const FVector Start, const FVector End) const = 0;
 
 };
\ No newline at end of file
