Fixed Mesh Paint Mode warning

diff --git a/Engine/Plugins/MeshPainting/Source/MeshPaintEditorMode/MeshPaintEditorMode.Build.cs b/Engine/Plugins/MeshPainting/Source/MeshPaintEditorMode/MeshPaintEditorMode.Build.cs
--- a/Engine/Plugins/MeshPainting/Source/MeshPaintEditorMode/MeshPaintEditorMode.Build.cs
+++ b/Engine/Plugins/MeshPainting/Source/MeshPaintEditorMode/MeshPaintEditorMode.Build.cs
@@ -23,6 +23,7 @@ namespace UnrealBuildTool.Rules
 				    "Engine",
                     "InputCore",
 					"EditorFramework",
+                    "EditorWidgets",
 					"UnrealEd",
 					"InteractiveToolsFramework",
 					"EditorInteractiveToolsFramework",
diff --git a/Engine/Plugins/MeshPainting/Source/MeshPaintEditorMode/Private/BrushSettingsCustomization.cpp b/Engine/Plugins/MeshPainting/Source/MeshPaintEditorMode/Private/BrushSettingsCustomization.cpp
--- a/Engine/Plugins/MeshPainting/Source/MeshPaintEditorMode/Private/BrushSettingsCustomization.cpp
+++ b/Engine/Plugins/MeshPainting/Source/MeshPaintEditorMode/Private/BrushSettingsCustomization.cpp
@@ -16,6 +16,7 @@
 #include "PropertyCustomizationHelpers.h"
 #include "Widgets/Input/SButton.h"
 #include "Widgets/Layout/SBox.h"
+#include "SWarningOrErrorBox.h"
 #include "ScopedTransaction.h"
 #include "MeshPaintHelpers.h"
 #include "MeshVertexPaintingTool.h"
@@ -289,7 +290,7 @@ void FColorPaintingSettingsCustomization::CustomizeDetails(IDetailLayoutBuilder&
 	ColorCategory.AddCustomRow(NSLOCTEXT("LODPainting", "LODPaintingLabel", "LOD Model Painting"))
 		.WholeRowContent()
 		[
-			SNew(SBorder)
+			SNew(SWarningOrErrorBox)
 			.Visibility_Lambda([this]() -> EVisibility
 			{
 				if (UMeshColorPaintingToolProperties* ColorProperties = UMeshPaintMode::GetColorToolProperties())
@@ -298,24 +299,15 @@ void FColorPaintingSettingsCustomization::CustomizeDetails(IDetailLayoutBuilder&
 				}
 				return EVisibility::Collapsed;
 			})
-		.Padding(FMargin(4.0f))
-		.BorderImage(FEditorStyle::GetBrush("SettingsEditor.CheckoutWarningBorder"))
-		.BorderBackgroundColor(FColor(166, 137, 0))
-		[
-			SNew(STextBlock)
-			.AutoWrapText(true)
-			.Font(IDetailLayoutBuilder::GetDetailFont())
-
-		.Text_Lambda([this]() -> FText
-		{
-			static const FText SkelMeshNotificationText = LOCTEXT("SkelMeshAssetPaintInfo", "Paint is propagated to Skeletal Mesh Asset(s)");
-			static const FText StaticMeshNotificationText = LOCTEXT("StaticMeshAssetPaintInfo", "Paint is applied to all LODs");
+			.Message_Lambda([this]() -> FText
+			{
+				static const FText SkelMeshNotificationText = LOCTEXT("SkelMeshAssetPaintInfo", "Paint is propagated to Skeletal Mesh Asset(s)");
+				static const FText StaticMeshNotificationText = LOCTEXT("StaticMeshAssetPaintInfo", "Paint is applied to all LODs");
 
-			const bool bSkelMeshText = UMeshPaintMode::GetMeshPaintMode()->GetSelectedComponents<USkeletalMeshComponent>().Num() > 0;
-			const bool bLODPaintText = UMeshPaintMode::GetColorToolProperties() ?  !UMeshPaintMode::GetColorToolProperties()->bPaintOnSpecificLOD : false;
-			return FText::Format(FTextFormat::FromString(TEXT("{0}{1}{2}")), bSkelMeshText ? SkelMeshNotificationText : FText::GetEmpty(), bSkelMeshText && bLODPaintText ? FText::FromString(TEXT("\n")) : FText::GetEmpty(), bLODPaintText ? StaticMeshNotificationText : FText::GetEmpty());
-		})
-		]
+				const bool bSkelMeshText = UMeshPaintMode::GetMeshPaintMode()->GetSelectedComponents<USkeletalMeshComponent>().Num() > 0;
+				const bool bLODPaintText = UMeshPaintMode::GetColorToolProperties() ?  !UMeshPaintMode::GetColorToolProperties()->bPaintOnSpecificLOD : false;
+				return FText::Format(FTextFormat::FromString(TEXT("{0}{1}{2}")), bSkelMeshText ? SkelMeshNotificationText : FText::GetEmpty(), bSkelMeshText && bLODPaintText ? FText::FromString(TEXT("\n")) : FText::GetEmpty(), bLODPaintText ? StaticMeshNotificationText : FText::GetEmpty());
+			})
 		];
 }
 
