Fix for crash when force-deleting assets that are selected for mesh painting

diff --git a/Engine/Plugins/MeshPainting/Source/MeshPaintEditorMode/Private/MeshPaintMode.cpp b/Engine/Plugins/MeshPainting/Source/MeshPaintEditorMode/Private/MeshPaintMode.cpp
--- a/Engine/Plugins/MeshPainting/Source/MeshPaintEditorMode/Private/MeshPaintMode.cpp
+++ b/Engine/Plugins/MeshPainting/Source/MeshPaintEditorMode/Private/MeshPaintMode.cpp
@@ -134,6 +134,7 @@ void UMeshPaintMode::Enter()
 {
 	Super::Enter();
 	GEditor->OnEditorClose().AddUObject(this, &UMeshPaintMode::OnResetViewMode);
+	GEditor->OnObjectsReplaced().AddUObject(this, &UMeshPaintMode::OnObjectsReplaced);
 	ModeSettings = Cast<UMeshPaintModeSettings>(SettingsObject);
 	
 	FMeshPaintEditorModeCommands ToolManagerCommands = FMeshPaintEditorModeCommands::Get();
@@ -171,7 +172,7 @@ void UMeshPaintMode::Enter()
 void UMeshPaintMode::Exit()
 {
 	Toolkit->OnPaletteChanged().Remove(PaletteChangedHandle);
-
+	GEditor->OnObjectsReplaced().RemoveAll(this);
 	GEditor->OnEditorClose().RemoveAll(this);
 	OnResetViewMode();
 	const FMeshPaintEditorModeCommands& Commands = FMeshPaintEditorModeCommands::Get();
@@ -439,6 +440,15 @@ void UMeshPaintMode::UpdateOnMaterialChange(bool bInvalidateHitProxies)
 	}
 }
 
+void UMeshPaintMode::OnObjectsReplaced(const TMap<UObject*, UObject*>& InOldToNewInstanceMap)
+{
+	if (UMeshPaintingSubsystem* MeshPaintingSubsystem = GEngine->GetEngineSubsystem<UMeshPaintingSubsystem>())
+	{
+		MeshPaintingSubsystem->ClearSelectedMeshComponents();
+		MeshPaintingSubsystem->Refresh();
+		UpdateSelectedMeshes();
+	}
+}
 
 void UMeshPaintMode::FillWithVertexColor()
 {
diff --git a/Engine/Plugins/MeshPainting/Source/MeshPaintEditorMode/Private/MeshPaintMode.h b/Engine/Plugins/MeshPainting/Source/MeshPaintEditorMode/Private/MeshPaintMode.h
--- a/Engine/Plugins/MeshPainting/Source/MeshPaintEditorMode/Private/MeshPaintMode.h
+++ b/Engine/Plugins/MeshPainting/Source/MeshPaintEditorMode/Private/MeshPaintMode.h
@@ -81,6 +81,7 @@ protected:
 	
 
 	void UpdateOnMaterialChange(bool bInvalidateHitProxies);
+	void OnObjectsReplaced(const TMap<UObject*, UObject*>& InOldToNewInstanceMap);
 	void OnResetViewMode();
 	void FillWithVertexColor();
 	void PropagateVertexColorsToAsset();
