Fix for deprecated OnObjectsReplaced binding

diff --git a/Engine/Plugins/MeshPainting/Source/MeshPaintEditorMode/Private/MeshPaintMode.cpp b/Engine/Plugins/MeshPainting/Source/MeshPaintEditorMode/Private/MeshPaintMode.cpp
--- a/Engine/Plugins/MeshPainting/Source/MeshPaintEditorMode/Private/MeshPaintMode.cpp
+++ b/Engine/Plugins/MeshPainting/Source/MeshPaintEditorMode/Private/MeshPaintMode.cpp
@@ -134,7 +134,7 @@ void UMeshPaintMode::Enter()
 {
 	Super::Enter();
 	GEditor->OnEditorClose().AddUObject(this, &UMeshPaintMode::OnResetViewMode);
-	GEditor->OnObjectsReplaced().AddUObject(this, &UMeshPaintMode::OnObjectsReplaced);
+	FCoreUObjectDelegates::OnObjectsReplaced.AddUObject(this, &UMeshPaintMode::OnObjectsReplaced);
 	ModeSettings = Cast<UMeshPaintModeSettings>(SettingsObject);
 	
 	FMeshPaintEditorModeCommands ToolManagerCommands = FMeshPaintEditorModeCommands::Get();
@@ -172,7 +172,7 @@ void UMeshPaintMode::Enter()
 void UMeshPaintMode::Exit()
 {
 	Toolkit->OnPaletteChanged().Remove(PaletteChangedHandle);
-	GEditor->OnObjectsReplaced().RemoveAll(this);
+	FCoreUObjectDelegates::OnObjectsReplaced.RemoveAll(this);
 	GEditor->OnEditorClose().RemoveAll(this);
 	OnResetViewMode();
 	const FMeshPaintEditorModeCommands& Commands = FMeshPaintEditorModeCommands::Get();
