Adding validity checks to texture paint

diff --git a/Engine/Plugins/MeshPainting/Source/MeshPaintingToolset/Private/TexturePaintToolset.cpp b/Engine/Plugins/MeshPainting/Source/MeshPaintingToolset/Private/TexturePaintToolset.cpp
--- a/Engine/Plugins/MeshPainting/Source/MeshPaintingToolset/Private/TexturePaintToolset.cpp
+++ b/Engine/Plugins/MeshPainting/Source/MeshPaintingToolset/Private/TexturePaintToolset.cpp
@@ -485,14 +485,17 @@ bool UTexturePaintToolset::DoesMeshComponentUseTexture(UMeshComponent* MeshCompo
 
 void UTexturePaintToolset::RetrieveTexturesForComponent(const UMeshComponent* Component, IMeshPaintComponentAdapter* Adapter, TArray<FPaintableTexture>& OutTextures)
 {
-	// Get the materials used by the mesh
-	TArray<UMaterialInterface*> UsedMaterials;
-	Component->GetUsedMaterials(UsedMaterials);
-
-	for (int32 MaterialIndex = 0; MaterialIndex < UsedMaterials.Num(); ++MaterialIndex)
+	if (Component && Adapter)
 	{
-		int32 OutDefaultIndex = 0;
-		Adapter->QueryPaintableTextures(MaterialIndex, OutDefaultIndex, OutTextures);
+		// Get the materials used by the mesh
+		TArray<UMaterialInterface*> UsedMaterials;
+		Component->GetUsedMaterials(UsedMaterials);
+
+		for (int32 MaterialIndex = 0; MaterialIndex < UsedMaterials.Num(); ++MaterialIndex)
+		{
+			int32 OutDefaultIndex = 0;
+			Adapter->QueryPaintableTextures(MaterialIndex, OutDefaultIndex, OutTextures);
+		}
 	}
 }
 
