Fixes for inconsistent mesh selection and painting behavior.

diff --git a/Engine/Plugins/MeshPainting/Source/MeshPaintEditorMode/Private/MeshPaintMode.cpp b/Engine/Plugins/MeshPainting/Source/MeshPaintEditorMode/Private/MeshPaintMode.cpp
--- a/Engine/Plugins/MeshPainting/Source/MeshPaintEditorMode/Private/MeshPaintMode.cpp
+++ b/Engine/Plugins/MeshPainting/Source/MeshPaintEditorMode/Private/MeshPaintMode.cpp
@@ -243,6 +243,11 @@ void UMeshPaintMode::Tick(FEditorViewportClient* ViewportClient, float DeltaTime
 	}
 }
 
+bool UMeshPaintMode::HandleClick(FEditorViewportClient* InViewportClient, HHitProxy* HitProxy, const FViewportClick& Click)
+{
+	return true;
+}
+
 TMap<FName, TArray<TSharedPtr<FUICommandInfo>>> UMeshPaintMode::GetModeCommands() const
 {
 	return FMeshPaintEditorModeCommands::GetCommands();
@@ -425,6 +430,11 @@ void UMeshPaintMode::CheckSelectionForTexturePaintCompat(const TArray<UMeshCompo
 	}
 
 	GEngine->GetEngineSubsystem<UMeshPaintingSubsystem>()->SetSelectionHasMaterialValidForTexturePaint(bCurrentSelectionSupportsTexturePaint);
+	if (!bCurrentSelectionSupportsTexturePaint && GetToolManager()->GetActiveTool(EToolSide::Mouse)->IsA<UMeshTexturePaintingTool>())
+	{
+		ToolsContext->EndTool(EToolShutdownType::Accept);
+		ToolsContext->StartTool(TextureSelectToolName);
+	}
 }
 
 void UMeshPaintMode::UpdateOnMaterialChange(bool bInvalidateHitProxies)
diff --git a/Engine/Plugins/MeshPainting/Source/MeshPaintEditorMode/Private/MeshPaintMode.h b/Engine/Plugins/MeshPainting/Source/MeshPaintEditorMode/Private/MeshPaintMode.h
--- a/Engine/Plugins/MeshPainting/Source/MeshPaintEditorMode/Private/MeshPaintMode.h
+++ b/Engine/Plugins/MeshPainting/Source/MeshPaintEditorMode/Private/MeshPaintMode.h
@@ -45,6 +45,7 @@ public:
 	virtual void Exit() override;
 	virtual void CreateToolkit() override;
 	virtual void Tick(FEditorViewportClient* ViewportClient, float DeltaTime) override;
+	virtual bool HandleClick(FEditorViewportClient* InViewportClient, HHitProxy* HitProxy, const FViewportClick& Click) override;
 	static FName MeshPaintMode_Color;
 	static FName MeshPaintMode_Texture;
 	static FName MeshPaintMode_Weights;
diff --git a/Engine/Plugins/MeshPainting/Source/MeshPaintingToolset/Private/MeshPaintHelpers.cpp b/Engine/Plugins/MeshPainting/Source/MeshPaintingToolset/Private/MeshPaintHelpers.cpp
--- a/Engine/Plugins/MeshPainting/Source/MeshPaintingToolset/Private/MeshPaintHelpers.cpp
+++ b/Engine/Plugins/MeshPainting/Source/MeshPaintingToolset/Private/MeshPaintHelpers.cpp
@@ -1396,21 +1396,39 @@ void UMeshPaintingSubsystem::ApplyVertexColorsToAllLODs(IMeshPaintComponentAdapt
 		}
 	}
 }
-TMap<FName, TSharedPtr<IMeshPaintComponentAdapter>> UMeshPaintingSubsystem::GetComponentToAdapterMap() const
-{
-	return ComponentToAdapterMap;
-}
 
-TSharedPtr<IMeshPaintComponentAdapter> UMeshPaintingSubsystem::GetAdapterForComponent(const UMeshComponent* InComponent)
+TSharedPtr<IMeshPaintComponentAdapter> UMeshPaintingSubsystem::GetAdapterForComponent(const UMeshComponent* InComponent) const
 {
-	FName ComponentName = InComponent->GetFName();
-	TSharedPtr<IMeshPaintComponentAdapter>* MeshAdapterPtr = ComponentToAdapterMap.Find(ComponentName);
+	FString ComponentActorString;
+	
+	if (const UStaticMeshComponent* InStaticMeshComponent = Cast<UStaticMeshComponent>(InComponent))
+	{
+		ComponentActorString = InStaticMeshComponent->GetStaticMesh()->GetName();
+	}
+	else if (const USkeletalMeshComponent* InSkeletalMeshComponent = Cast<USkeletalMeshComponent>(InComponent))
+	{
+		ComponentActorString = InSkeletalMeshComponent->SkeletalMesh->GetName();
+	} 
+	ComponentActorString += TEXT("_") + InComponent->GetOwner()->GetName();
+	const FName ComponentActorName = FName(ComponentActorString);
+	const TSharedPtr<IMeshPaintComponentAdapter>* MeshAdapterPtr = ComponentToAdapterMap.Find(ComponentActorName);
 	return MeshAdapterPtr ? *MeshAdapterPtr : TSharedPtr<IMeshPaintComponentAdapter>();
 }
 
-void UMeshPaintingSubsystem::AddToComponentToAdapterMap(UMeshComponent* InComponent, TSharedPtr<IMeshPaintComponentAdapter> InAdapter)
+void UMeshPaintingSubsystem::AddToComponentToAdapterMap(const UMeshComponent* InComponent, const TSharedPtr<IMeshPaintComponentAdapter> InAdapter) 
 {
-	ComponentToAdapterMap.Add(InComponent->GetFName(), InAdapter);
+	FString ComponentActorString;
+	if (const UStaticMeshComponent* InStaticMeshComponent = Cast<UStaticMeshComponent>(InComponent))
+	{
+		ComponentActorString = InStaticMeshComponent->GetStaticMesh()->GetName();
+	}
+	else if (const USkeletalMeshComponent* InSkeletalMeshComponent = Cast<USkeletalMeshComponent>(InComponent))
+	{
+		ComponentActorString = InSkeletalMeshComponent->SkeletalMesh->GetName();
+	}
+	ComponentActorString += TEXT("_") + InComponent->GetOwner()->GetName();
+	const FName ComponentActorName = FName(ComponentActorString);
+	ComponentToAdapterMap.Add(ComponentActorName, InAdapter);
 }
 
 TArray<UMeshComponent*> UMeshPaintingSubsystem::GetSelectedMeshComponents() const
@@ -1498,7 +1516,7 @@ bool UMeshPaintingSubsystem::FindHitResult(const FRay Ray, FHitResult& BestTrace
 			}
 		}
 	}
-	return true;
+	return BestTraceResult.Distance != FLT_MAX;
 }
 
 bool UMeshPaintingSubsystem::SelectionContainsValidAdapters() const
diff --git a/Engine/Plugins/MeshPainting/Source/MeshPaintingToolset/Private/MeshPaintInteractions.cpp b/Engine/Plugins/MeshPainting/Source/MeshPaintingToolset/Private/MeshPaintInteractions.cpp
--- a/Engine/Plugins/MeshPainting/Source/MeshPaintingToolset/Private/MeshPaintInteractions.cpp
+++ b/Engine/Plugins/MeshPainting/Source/MeshPaintingToolset/Private/MeshPaintInteractions.cpp
@@ -19,7 +19,7 @@
 
 #define LOCTEXT_NAMESPACE "MeshSelection"
 
-FInputRayHit UMeshPaintSelectionMechanic::IsHitByClick(const FInputDeviceRay& ClickPos)
+FInputRayHit UMeshPaintSelectionMechanic::IsHitByClick(const FInputDeviceRay& ClickPos, bool bIsFallbackClick)
 {
 	IMeshPaintSelectionInterface* Interface = Cast<IMeshPaintSelectionInterface>(GetParentTool());
 	if (!bAddToSelectionSet || !Interface->AllowsMultiselect())
@@ -27,6 +27,19 @@ FInputRayHit UMeshPaintSelectionMechanic::IsHitByClick(const FInputDeviceRay& Cl
 		CachedClickedComponents.Empty();
 		CachedClickedActors.Empty();
 	}
+
+	// for fallback clicks, assume we must be adding to our selection set
+	if (bIsFallbackClick)
+	{
+		if (Interface->AllowsMultiselect() && bAddToSelectionSet)
+		{
+			return FindClickedComponentsAndCacheAdapters(ClickPos) ? FInputRayHit(0.0f) : FInputRayHit();
+		}
+		else
+		{
+			return FInputRayHit();
+		}
+	}
 	return FindClickedComponentsAndCacheAdapters(ClickPos) ? FInputRayHit(0.0f) : FInputRayHit();
 }
 
@@ -53,7 +66,8 @@ void UMeshPaintSelectionMechanic::OnClicked(const FInputDeviceRay& ClickPos)
 
 			FSelectedOjectsChangeList NewSelection;
 			// TODO add CTRL handling
-			NewSelection.ModificationType = bAddToSelectionSet ? ESelectedObjectsModificationType::Add : ESelectedObjectsModificationType::Replace;
+			const bool bShouldAddToSelection = bAddToSelectionSet && Interface->AllowsMultiselect();
+			NewSelection.ModificationType = bShouldAddToSelection ? ESelectedObjectsModificationType::Add : ESelectedObjectsModificationType::Replace;
 			NewSelection.Actors.Append(CachedClickedActors);
 			GetParentTool()->GetToolManager()->RequestSelectionChange(NewSelection);
 			GetParentTool()->GetToolManager()->EndUndoTransaction();
@@ -78,7 +92,7 @@ bool UMeshPaintSelectionMechanic::FindClickedComponentsAndCacheAdapters(const FI
 		{
 			if (TTypedElement<UTypedElementObjectInterface> ObjectInterface = UTypedElementRegistry::GetInstance()->GetElement<UTypedElementObjectInterface>(HitProxy->GetElementHandle()))
 			{
-				if (AActor* Actor = ObjectInterface.GetObjectAs<AActor>())
+				if (AActor* Actor = ObjectInterface.GetObjectAs<UActorComponent>()->GetOwner())
 				{
 					TArray<UActorComponent*> CandidateComponents = Actor->K2_GetComponentsByClass(UMeshComponent::StaticClass());
 					for (UActorComponent* CandidateComponent : CandidateComponents)
diff --git a/Engine/Plugins/MeshPainting/Source/MeshPaintingToolset/Private/MeshTexturePaintingTool.cpp b/Engine/Plugins/MeshPainting/Source/MeshPaintingToolset/Private/MeshTexturePaintingTool.cpp
--- a/Engine/Plugins/MeshPainting/Source/MeshPaintingToolset/Private/MeshTexturePaintingTool.cpp
+++ b/Engine/Plugins/MeshPainting/Source/MeshPaintingToolset/Private/MeshTexturePaintingTool.cpp
@@ -73,11 +73,21 @@ void UMeshTexturePaintingTool::Setup()
 	bStampPending = false;
 	BrushProperties->RestoreProperties(this);
 
+	// Needed after restoring properties because the brush radius may be an output
+	// property based on selection, so we shouldn't use the last stored value there.
+	// We wouldn't have this problem if we restore properties before getting
+	// BrushRelativeSizeRange, but that happens in the Super::Setup() call earlier.
+	RecalculateBrushRadius();
+
+	BrushStampIndicator->LineColor = FLinearColor::Green;
 
 	GetToolManager()->DisplayMessage(
 		LOCTEXT("OnStartTexturePaintTool", "The Texture Weight Painting mode enables you to paint on textures and access available properties while doing so ."),
 		EToolMessageLevel::UserNotification);
 
+	SelectionMechanic = NewObject<UMeshPaintSelectionMechanic>(this);
+	SelectionMechanic->Setup(this);
+
 }
 
 
@@ -482,6 +492,31 @@ bool UMeshTexturePaintingTool::CanAccept() const
 	return false;
 }
 
+FInputRayHit UMeshTexturePaintingTool::CanBeginClickDragSequence(const FInputDeviceRay& PressPos)
+{
+	FHitResult OutHit;
+	bCachedClickRay = false;
+	if (!HitTest(PressPos.WorldRay, OutHit))
+	{
+		UMeshPaintingSubsystem* MeshPaintingSubsystem = GEngine->GetEngineSubsystem<UMeshPaintingSubsystem>();
+		const bool bFallbackClick = MeshPaintingSubsystem->GetSelectedMeshComponents().Num() > 0;
+		if (SelectionMechanic->IsHitByClick(PressPos, bFallbackClick).bHit)
+		{
+			bCachedClickRay = true;
+			PendingClickRay = PressPos.WorldRay;
+			PendingClickScreenPosition = PressPos.ScreenPosition;
+			return FInputRayHit(0.0);
+		}
+	}
+	return Super::CanBeginClickDragSequence(PressPos);
+}
+
+void UMeshTexturePaintingTool::OnUpdateModifierState(int ModifierID, bool bIsOn)
+{
+	Super::OnUpdateModifierState(ModifierID, bIsOn);
+	SelectionMechanic->SetAddToSelectionSet(bShiftToggle);
+}
+
 void UMeshTexturePaintingTool::OnBeginDrag(const FRay& Ray)
 {
 	Super::OnBeginDrag(Ray);
@@ -494,6 +529,14 @@ void UMeshTexturePaintingTool::OnBeginDrag(const FRay& Ray)
 		PendingStampRay = Ray;
 		bStampPending = true;
 	}
+	else if (bCachedClickRay)
+	{
+		FInputDeviceRay InputDeviceRay = FInputDeviceRay(PendingClickRay, PendingClickScreenPosition);
+		SelectionMechanic->SetAddToSelectionSet(bShiftToggle);
+		SelectionMechanic->OnClicked(InputDeviceRay);
+		bCachedClickRay = false;
+		RecalculateBrushRadius();
+	}
 }
 
 void UMeshTexturePaintingTool::OnUpdateDrag(const FRay& Ray)
@@ -1295,4 +1338,10 @@ void UMeshTexturePaintingTool::PaintTextureChanged(const FAssetData& AssetData)
 		}
 	}
 }
+
+bool UMeshTexturePaintingTool::IsMeshAdapterSupported(TSharedPtr<IMeshPaintComponentAdapter> MeshAdapter) const
+{
+	return MeshAdapter.IsValid() ? MeshAdapter->SupportsTexturePaint() : false;
+}
+
 #undef LOCTEXT_NAMESPACE
diff --git a/Engine/Plugins/MeshPainting/Source/MeshPaintingToolset/Private/MeshVertexPaintingTool.cpp b/Engine/Plugins/MeshPainting/Source/MeshPaintingToolset/Private/MeshVertexPaintingTool.cpp
--- a/Engine/Plugins/MeshPainting/Source/MeshPaintingToolset/Private/MeshVertexPaintingTool.cpp
+++ b/Engine/Plugins/MeshPainting/Source/MeshPaintingToolset/Private/MeshVertexPaintingTool.cpp
@@ -82,7 +82,7 @@ void UMeshVertexPaintingTool::Setup()
 	RecalculateBrushRadius(); 
 
 	BrushStampIndicator->LineColor = FLinearColor::Green;
-	// Set up selection mechanic to find and select edges
+
 	SelectionMechanic = NewObject<UMeshPaintSelectionMechanic>(this);
 	SelectionMechanic->Setup(this);
 }
@@ -125,7 +125,7 @@ void UMeshVertexPaintingTool::Render(IToolsContextRenderAPI* RenderAPI)
 		{
 			TSharedPtr<IMeshPaintComponentAdapter> MeshAdapter = MeshPaintingSubsystem->GetAdapterForComponent(Cast<UMeshComponent>(CurrentComponent));
 
-			if (MeshAdapter->IsValid() && MeshAdapter->SupportsVertexPaint())
+			if (MeshAdapter && MeshAdapter->SupportsVertexPaint())
 			{
 				const FMatrix ComponentToWorldMatrix = MeshAdapter->GetComponentToWorldMatrix();
 				FViewCameraState CameraState;
@@ -302,8 +302,8 @@ bool UMeshVertexPaintingTool::PaintInternal(const TArrayView<TPair<FVector, FVec
 			for (auto TestComponent : MeshPaintingSubsystem->GetPaintableMeshComponents())
 			{
 				const FBox ComponentBounds = TestComponent->Bounds.GetBox();
-
-				if (MeshPaintingSubsystem->GetComponentToAdapterMap().Contains(TestComponent->GetFName()) && ComponentBounds.Intersect(BrushBounds))
+				TSharedPtr<IMeshPaintComponentAdapter> MeshAdapter = MeshPaintingSubsystem->GetAdapterForComponent(TestComponent);
+				if (MeshAdapter && ComponentBounds.Intersect(BrushBounds))
 				{
 					// OK, this mesh potentially overlaps the brush!
 					HoveredComponents.FindOrAdd(TestComponent).Add(i);
@@ -454,7 +454,9 @@ FInputRayHit UMeshVertexPaintingTool::CanBeginClickDragSequence(const FInputDevi
 	bCachedClickRay = false;
 	if (!HitTest(PressPos.WorldRay, OutHit))
 	{
-		if (SelectionMechanic->IsHitByClick(PressPos).bHit)
+		UMeshPaintingSubsystem* MeshPaintingSubsystem = GEngine->GetEngineSubsystem<UMeshPaintingSubsystem>();
+		const bool bFallbackClick = MeshPaintingSubsystem->GetSelectedMeshComponents().Num() > 0;
+		if (SelectionMechanic->IsHitByClick(PressPos, bFallbackClick).bHit)
 		{
 			bCachedClickRay = true;
 			PendingClickRay = PressPos.WorldRay;
diff --git a/Engine/Plugins/MeshPainting/Source/MeshPaintingToolset/Public/MeshPaintHelpers.h b/Engine/Plugins/MeshPainting/Source/MeshPaintingToolset/Public/MeshPaintHelpers.h
--- a/Engine/Plugins/MeshPainting/Source/MeshPaintingToolset/Public/MeshPaintHelpers.h
+++ b/Engine/Plugins/MeshPainting/Source/MeshPaintingToolset/Public/MeshPaintHelpers.h
@@ -242,9 +242,8 @@ public:
 	}
 
 	/** Map of geometry adapters for each selected mesh component */
-	TMap<FName, TSharedPtr<IMeshPaintComponentAdapter>> GetComponentToAdapterMap() const;
-	TSharedPtr<IMeshPaintComponentAdapter> GetAdapterForComponent(const UMeshComponent* InComponent);
-	void AddToComponentToAdapterMap(UMeshComponent* InComponent, TSharedPtr<IMeshPaintComponentAdapter> InAdapter);
+	TSharedPtr<IMeshPaintComponentAdapter> GetAdapterForComponent(const UMeshComponent* InComponent) const;
+	void AddToComponentToAdapterMap(const UMeshComponent* InComponent, const TSharedPtr<IMeshPaintComponentAdapter> InAdapter);
 
 	TArray<UMeshComponent*> GetSelectedMeshComponents() const;
 	void AddSelectedMeshComponents(const TArray<UMeshComponent*>& InComponents);
diff --git a/Engine/Plugins/MeshPainting/Source/MeshPaintingToolset/Public/MeshPaintInteractions.h b/Engine/Plugins/MeshPainting/Source/MeshPaintingToolset/Public/MeshPaintInteractions.h
--- a/Engine/Plugins/MeshPainting/Source/MeshPaintingToolset/Public/MeshPaintInteractions.h
+++ b/Engine/Plugins/MeshPainting/Source/MeshPaintingToolset/Public/MeshPaintInteractions.h
@@ -33,7 +33,7 @@ class MESHPAINTINGTOOLSET_API UMeshPaintSelectionMechanic : public UInteractionM
 	GENERATED_BODY()
 
 public:
-	FInputRayHit IsHitByClick(const FInputDeviceRay& ClickPos);
+	FInputRayHit IsHitByClick(const FInputDeviceRay& ClickPos, bool bIsFallbackClick = false);
 	void OnClicked(const FInputDeviceRay& ClickPos);
 	void SetAddToSelectionSet(const bool bInNewSelectionType)
 	{
diff --git a/Engine/Plugins/MeshPainting/Source/MeshPaintingToolset/Public/MeshTexturePaintingTool.h b/Engine/Plugins/MeshPainting/Source/MeshPaintingToolset/Public/MeshTexturePaintingTool.h
--- a/Engine/Plugins/MeshPainting/Source/MeshPaintingToolset/Public/MeshTexturePaintingTool.h
+++ b/Engine/Plugins/MeshPainting/Source/MeshPaintingToolset/Public/MeshTexturePaintingTool.h
@@ -96,7 +96,7 @@ public:
 
 
 UCLASS()
-class MESHPAINTINGTOOLSET_API UMeshTexturePaintingTool : public UBaseBrushTool
+class MESHPAINTINGTOOLSET_API UMeshTexturePaintingTool : public UBaseBrushTool, public IMeshPaintSelectionInterface
 {
 	GENERATED_BODY()
 
@@ -110,6 +110,8 @@ public:
 	virtual bool HasCancel() const override { return false; }
 	virtual bool HasAccept() const override;
 	virtual bool CanAccept() const override;
+	virtual FInputRayHit CanBeginClickDragSequence(const FInputDeviceRay& PressPos) override;
+	virtual void OnUpdateModifierState(int ModifierID, bool bIsOn) override;
 	virtual void OnBeginDrag(const FRay& Ray) override;
 	virtual void OnUpdateDrag(const FRay& Ray) override;
 	virtual void OnEndDrag(const FRay& Ray) override;
@@ -135,6 +137,11 @@ public:
 
 	bool ShouldFilterTextureAsset(const FAssetData& AssetData) const;
 	void PaintTextureChanged(const FAssetData& AssetData);
+	virtual bool IsMeshAdapterSupported(TSharedPtr<IMeshPaintComponentAdapter> MeshAdapter) const override;
+	virtual bool AllowsMultiselect() const override
+	{
+		return false;
+	}
 
 protected:
 	virtual void SetAdditionalPaintParameters(FMeshPaintParameters& InPaintParameters) {};
@@ -158,6 +165,12 @@ protected:
 	bool bStampPending;
 	bool bInDrag;
 	FRay PendingStampRay;
+	FRay PendingClickRay;
+	FVector2D PendingClickScreenPosition;
+	bool bCachedClickRay;
+
+	UPROPERTY(Transient)
+	TObjectPtr<UMeshPaintSelectionMechanic> SelectionMechanic;
 
 private:
 	bool PaintInternal(const TArrayView<TPair<FVector, FVector>>& Rays, EMeshPaintModeAction PaintAction, float PaintStrength);
diff --git a/Engine/Source/Editor/Experimental/EditorInteractiveToolsFramework/Private/EdModeInteractiveToolsContext.cpp b/Engine/Source/Editor/Experimental/EditorInteractiveToolsFramework/Private/EdModeInteractiveToolsContext.cpp
--- a/Engine/Source/Editor/Experimental/EditorInteractiveToolsFramework/Private/EdModeInteractiveToolsContext.cpp
+++ b/Engine/Source/Editor/Experimental/EditorInteractiveToolsFramework/Private/EdModeInteractiveToolsContext.cpp
@@ -955,18 +955,21 @@ bool UEdModeInteractiveToolsContext::InputKey(FEditorViewportClient* ViewportCli
 					InputState.Mouse.Left.SetStates(
 						(Event == IE_Pressed), (Event == IE_Pressed), (Event == IE_Released));
 					CurrentMouseState.Mouse.Left.bDown = (Event == IE_Pressed);
+					CurrentMouseState.Mouse.Left.bReleased = (Event == IE_Released);
 				}
 				else if (bIsMiddleMouse)
 				{
 					InputState.Mouse.Middle.SetStates(
 						(Event == IE_Pressed), (Event == IE_Pressed), (Event == IE_Released));
 					CurrentMouseState.Mouse.Middle.bDown = (Event == IE_Pressed);
+					CurrentMouseState.Mouse.Middle.bReleased = (Event == IE_Released);
 				}
 				else
 				{
 					InputState.Mouse.Right.SetStates(
 						(Event == IE_Pressed), (Event == IE_Pressed), (Event == IE_Released));
 					CurrentMouseState.Mouse.Right.bDown = (Event == IE_Pressed);
+					CurrentMouseState.Mouse.Right.bReleased = (Event == IE_Released);
 				}
 				if (InputRouter->PostInputEvent(InputState))
 				{
