Adding shared pointer validity condition for mesh paint adapter

diff --git a/Engine/Plugins/MeshPainting/Source/MeshPaintingToolset/Private/MeshPaintHelpers.cpp b/Engine/Plugins/MeshPainting/Source/MeshPaintingToolset/Private/MeshPaintHelpers.cpp
--- a/Engine/Plugins/MeshPainting/Source/MeshPaintingToolset/Private/MeshPaintHelpers.cpp
+++ b/Engine/Plugins/MeshPainting/Source/MeshPaintingToolset/Private/MeshPaintHelpers.cpp
@@ -1523,7 +1523,7 @@ bool UMeshPaintingSubsystem::SelectionContainsValidAdapters() const
 {
 	for (auto& MeshAdapterPair : ComponentToAdapterMap)
 	{
-		if (MeshAdapterPair.Value->IsValid())
+		if (MeshAdapterPair.Value && MeshAdapterPair.Value->IsValid())
 		{
 			return true;
 		}
