Keeping the state of the mesh paint adapter map much more consistent, to avoid issues when undoing and going into PIE

diff --git a/Engine/Plugins/MeshPainting/Source/MeshPaintEditorMode/Private/MeshPaintMode.cpp b/Engine/Plugins/MeshPainting/Source/MeshPaintEditorMode/Private/MeshPaintMode.cpp
--- a/Engine/Plugins/MeshPainting/Source/MeshPaintEditorMode/Private/MeshPaintMode.cpp
+++ b/Engine/Plugins/MeshPainting/Source/MeshPaintEditorMode/Private/MeshPaintMode.cpp
@@ -163,7 +163,7 @@ void UMeshPaintMode::Enter()
 	// disable tool change tracking to activate default tool, and then switch to full undo/redo tracking mode
 	GetToolManager()->ConfigureChangeTrackingMode(EToolChangeTrackingMode::NoChangeTracking);
 	Toolkit->SetCurrentPalette(MeshPaintMode_Color);
-	GetToolManager()->ConfigureChangeTrackingMode(EToolChangeTrackingMode::FullUndoRedo);
+	GetToolManager()->ConfigureChangeTrackingMode(EToolChangeTrackingMode::UndoToExit);
 
 	FLevelEditorModule& LevelEditor = FModuleManager::GetModuleChecked<FLevelEditorModule>(FName(TEXT("LevelEditor")));
 	LevelEditor.OnRedrawLevelEditingViewports().AddUObject(this, &UMeshPaintMode::UpdateOnMaterialChange);
@@ -422,6 +422,10 @@ void UMeshPaintMode::CheckSelectionForTexturePaintCompat(const TArray<UMeshCompo
 				if (!MeshAdapter.IsValid())
 				{
 					MeshAdapter = FMeshPaintComponentAdapterFactory::CreateAdapterForMesh(MeshComponent, 0);
+					if (MeshAdapter)
+					{
+						GEngine->GetEngineSubsystem<UMeshPaintingSubsystem>()->AddToComponentToAdapterMap(MeshComponent, MeshAdapter);
+					}
 				}
 				UTexturePaintToolset::RetrieveTexturesForComponent(MeshComponent, MeshAdapter.Get(), PaintableTextures);
 			}
@@ -430,7 +434,10 @@ void UMeshPaintMode::CheckSelectionForTexturePaintCompat(const TArray<UMeshCompo
 	}
 
 	GEngine->GetEngineSubsystem<UMeshPaintingSubsystem>()->SetSelectionHasMaterialValidForTexturePaint(bCurrentSelectionSupportsTexturePaint);
-	if (!bCurrentSelectionSupportsTexturePaint && GetToolManager()->GetActiveTool(EToolSide::Mouse)->IsA<UMeshTexturePaintingTool>())
+	if (!bCurrentSelectionSupportsTexturePaint 
+		&& GetToolManager() 
+		&& GetToolManager()->GetActiveTool(EToolSide::Mouse)
+		&& GetToolManager()->GetActiveTool(EToolSide::Mouse)->IsA<UMeshTexturePaintingTool>())
 	{
 		ToolsContext->EndTool(EToolShutdownType::Accept);
 		ToolsContext->StartTool(TextureSelectToolName);
diff --git a/Engine/Plugins/MeshPainting/Source/MeshPaintingToolset/Private/MeshPaintHelpers.cpp b/Engine/Plugins/MeshPainting/Source/MeshPaintingToolset/Private/MeshPaintHelpers.cpp
--- a/Engine/Plugins/MeshPainting/Source/MeshPaintingToolset/Private/MeshPaintHelpers.cpp
+++ b/Engine/Plugins/MeshPainting/Source/MeshPaintingToolset/Private/MeshPaintHelpers.cpp
@@ -1466,7 +1466,6 @@ void UMeshPaintingSubsystem::ClearPaintableMeshComponents()
 void UMeshPaintingSubsystem::ResetState()
 {
 	PaintableComponents.Empty();
-	ComponentToAdapterMap.Empty();
 	SelectedMeshComponents.Empty();
 }
 
diff --git a/Engine/Plugins/MeshPainting/Source/MeshPaintingToolset/Private/MeshTexturePaintingTool.cpp b/Engine/Plugins/MeshPainting/Source/MeshPaintingToolset/Private/MeshTexturePaintingTool.cpp
--- a/Engine/Plugins/MeshPainting/Source/MeshPaintingToolset/Private/MeshTexturePaintingTool.cpp
+++ b/Engine/Plugins/MeshPainting/Source/MeshPaintingToolset/Private/MeshTexturePaintingTool.cpp
@@ -95,7 +95,7 @@ void UMeshTexturePaintingTool::Shutdown(EToolShutdownType ShutdownType)
 {
 	FinishPainting();
 	// If the user has pending changes and the editor is not exiting, we want to do the commit for all the modified textures.
-	if ((GetNumberOfPendingPaintChanges() > 0) && !IsEngineExitRequested())
+	if ((GetNumberOfPendingPaintChanges() > 0) && !IsEngineExitRequested() && ShutdownType != EToolShutdownType::Cancel)
 	{
 		CommitAllPaintedTextures();
 	}
@@ -553,8 +553,8 @@ void UMeshTexturePaintingTool::OnUpdateDrag(const FRay& Ray)
 
 void UMeshTexturePaintingTool::OnEndDrag(const FRay& Ray)
 {
-	FinishPainting();
 	FinishPaintingTexture();
+	FinishPainting();
 	bStampPending = false;
 	bInDrag = false;
 }
