Fix a crash in mesh painting mode when selecting mesh component types which do not have a valid adapter implementation.

diff --git a/Engine/Plugins/MeshPainting/Source/MeshPaintingToolset/Private/MeshPaintHelpers.cpp b/Engine/Plugins/MeshPainting/Source/MeshPaintingToolset/Private/MeshPaintHelpers.cpp
--- a/Engine/Plugins/MeshPainting/Source/MeshPaintingToolset/Private/MeshPaintHelpers.cpp
+++ b/Engine/Plugins/MeshPainting/Source/MeshPaintingToolset/Private/MeshPaintHelpers.cpp
@@ -1409,6 +1409,10 @@ TSharedPtr<IMeshPaintComponentAdapter> UMeshPaintingSubsystem::GetAdapterForComp
 	{
 		ComponentActorString = InSkeletalMeshComponent->SkeletalMesh->GetName();
 	} 
+	else
+	{
+		return TSharedPtr<IMeshPaintComponentAdapter>();
+	}
 	ComponentActorString += TEXT("_") + InComponent->GetOwner()->GetName();
 	const FName ComponentActorName = FName(ComponentActorString);
 	const TSharedPtr<IMeshPaintComponentAdapter>* MeshAdapterPtr = ComponentToAdapterMap.Find(ComponentActorName);
@@ -1417,6 +1421,11 @@ TSharedPtr<IMeshPaintComponentAdapter> UMeshPaintingSubsystem::GetAdapterForComp
 
 void UMeshPaintingSubsystem::AddToComponentToAdapterMap(const UMeshComponent* InComponent, const TSharedPtr<IMeshPaintComponentAdapter> InAdapter) 
 {
+	if (!InAdapter)
+	{
+		return;
+	}
+
 	FString ComponentActorString;
 	if (const UStaticMeshComponent* InStaticMeshComponent = Cast<UStaticMeshComponent>(InComponent))
 	{
@@ -1426,6 +1435,10 @@ void UMeshPaintingSubsystem::AddToComponentToAdapterMap(const UMeshComponent* In
 	{
 		ComponentActorString = InSkeletalMeshComponent->SkeletalMesh->GetName();
 	}
+	else
+	{
+		ComponentActorString = InComponent->GetName();
+	}
 	ComponentActorString += TEXT("_") + InComponent->GetOwner()->GetName();
 	const FName ComponentActorName = FName(ComponentActorString);
 	ComponentToAdapterMap.Add(ComponentActorName, InAdapter);
diff --git a/Engine/Plugins/MeshPainting/Source/MeshPaintingToolset/Private/MeshPaintInteractions.cpp b/Engine/Plugins/MeshPainting/Source/MeshPaintingToolset/Private/MeshPaintInteractions.cpp
--- a/Engine/Plugins/MeshPainting/Source/MeshPaintingToolset/Private/MeshPaintInteractions.cpp
+++ b/Engine/Plugins/MeshPainting/Source/MeshPaintingToolset/Private/MeshPaintInteractions.cpp
@@ -102,7 +102,10 @@ bool UMeshPaintSelectionMechanic::FindClickedComponentsAndCacheAdapters(const FI
 						if (!MeshAdapter.IsValid())
 						{
 							MeshAdapter = FMeshPaintComponentAdapterFactory::CreateAdapterForMesh(MeshComponent, 0);
-							MeshPaintingSubsystem->AddToComponentToAdapterMap(MeshComponent, MeshAdapter);
+							if (MeshAdapter)
+							{
+								MeshPaintingSubsystem->AddToComponentToAdapterMap(MeshComponent, MeshAdapter);
+							}
 						}
 						IMeshPaintSelectionInterface* Interface = Cast<IMeshPaintSelectionInterface>(GetParentTool());
 						if (MeshAdapter.IsValid() && Interface->IsMeshAdapterSupported(MeshAdapter))
