Refactored the mesh paint component adapter map to use the object path name as the key instead of trying to come up with our own unique names for components.

diff --git a/Engine/Plugins/MeshPainting/Source/MeshPaintingToolset/Private/MeshPaintHelpers.cpp b/Engine/Plugins/MeshPainting/Source/MeshPaintingToolset/Private/MeshPaintHelpers.cpp
--- a/Engine/Plugins/MeshPainting/Source/MeshPaintingToolset/Private/MeshPaintHelpers.cpp
+++ b/Engine/Plugins/MeshPainting/Source/MeshPaintingToolset/Private/MeshPaintHelpers.cpp
@@ -1399,49 +1399,17 @@ void UMeshPaintingSubsystem::ApplyVertexColorsToAllLODs(IMeshPaintComponentAdapt
 
 TSharedPtr<IMeshPaintComponentAdapter> UMeshPaintingSubsystem::GetAdapterForComponent(const UMeshComponent* InComponent) const
 {
-	FString ComponentActorString;
-	
-	if (const UStaticMeshComponent* InStaticMeshComponent = Cast<UStaticMeshComponent>(InComponent))
-	{
-		ComponentActorString = InStaticMeshComponent->GetStaticMesh()->GetName();
-	}
-	else if (const USkeletalMeshComponent* InSkeletalMeshComponent = Cast<USkeletalMeshComponent>(InComponent))
-	{
-		ComponentActorString = InSkeletalMeshComponent->SkeletalMesh->GetName();
-	} 
-	else
-	{
-		return TSharedPtr<IMeshPaintComponentAdapter>();
-	}
-	ComponentActorString += TEXT("_") + InComponent->GetOwner()->GetName();
-	const FName ComponentActorName = FName(ComponentActorString);
-	const TSharedPtr<IMeshPaintComponentAdapter>* MeshAdapterPtr = ComponentToAdapterMap.Find(ComponentActorName);
-	return MeshAdapterPtr ? *MeshAdapterPtr : TSharedPtr<IMeshPaintComponentAdapter>();
+	return InComponent ? ComponentToAdapterMap.FindRef(InComponent->GetPathName()) : TSharedPtr<IMeshPaintComponentAdapter>{};
 }
 
 void UMeshPaintingSubsystem::AddToComponentToAdapterMap(const UMeshComponent* InComponent, const TSharedPtr<IMeshPaintComponentAdapter> InAdapter) 
 {
-	if (!InAdapter)
+	if (!InAdapter || !InComponent)
 	{
 		return;
 	}
 
-	FString ComponentActorString;
-	if (const UStaticMeshComponent* InStaticMeshComponent = Cast<UStaticMeshComponent>(InComponent))
-	{
-		ComponentActorString = InStaticMeshComponent->GetStaticMesh()->GetName();
-	}
-	else if (const USkeletalMeshComponent* InSkeletalMeshComponent = Cast<USkeletalMeshComponent>(InComponent))
-	{
-		ComponentActorString = InSkeletalMeshComponent->SkeletalMesh->GetName();
-	}
-	else
-	{
-		ComponentActorString = InComponent->GetName();
-	}
-	ComponentActorString += TEXT("_") + InComponent->GetOwner()->GetName();
-	const FName ComponentActorName = FName(ComponentActorString);
-	ComponentToAdapterMap.Add(ComponentActorName, InAdapter);
+	ComponentToAdapterMap.Add(InComponent->GetPathName(), InAdapter);
 }
 
 TArray<UMeshComponent*> UMeshPaintingSubsystem::GetSelectedMeshComponents() const
diff --git a/Engine/Plugins/MeshPainting/Source/MeshPaintingToolset/Public/MeshPaintHelpers.h b/Engine/Plugins/MeshPainting/Source/MeshPaintingToolset/Public/MeshPaintHelpers.h
--- a/Engine/Plugins/MeshPainting/Source/MeshPaintingToolset/Public/MeshPaintHelpers.h
+++ b/Engine/Plugins/MeshPainting/Source/MeshPaintingToolset/Public/MeshPaintHelpers.h
@@ -273,7 +273,7 @@ private:
 
 private:
 	/** Map of geometry adapters for each selected mesh component */
-	TMap<FName, TSharedPtr<IMeshPaintComponentAdapter>> ComponentToAdapterMap;
+	TMap<FString, TSharedPtr<IMeshPaintComponentAdapter>> ComponentToAdapterMap;
 	TArray<UMeshComponent*> SelectedMeshComponents;
 	/** Mesh components within the current selection which are eligible for painting */
 	TArray<UMeshComponent*> PaintableComponents;
