Refactored the mesh paint component adapter map to use the object path name as the key instead of trying to come up with our own unique names for components.

diff --git a/Engine/Plugins/MeshPainting/Source/MeshPaintingToolset/Private/MeshPaintHelpers.cpp b/Engine/Plugins/MeshPainting/Source/MeshPaintingToolset/Private/MeshPaintHelpers.cpp
--- a/Engine/Plugins/MeshPainting/Source/MeshPaintingToolset/Private/MeshPaintHelpers.cpp
+++ b/Engine/Plugins/MeshPainting/Source/MeshPaintingToolset/Private/MeshPaintHelpers.cpp
@@ -1398,76 +1398,19 @@ void UMeshPaintingSubsystem::ApplyVertexColorsToAllLODs(IMeshPaintComponentAdapt
 	}
 }
 
-FName UMeshPaintingSubsystem::GetComponentActorName(const UMeshComponent* InComponent) const
-{
-	check(InComponent);
-
-	FString ComponentActorString;
-
-	// Must check USplineMeshComponent before UStaticMeshComponent since the former derives from the latter.
-	if (InComponent->IsA(USplineMeshComponent::StaticClass()))
-	{
-		ComponentActorString = InComponent->GetName();
-	}
-	else if (const UStaticMeshComponent* InStaticMeshComponent = Cast<UStaticMeshComponent>(InComponent))
-	{
-		if (InStaticMeshComponent->GetStaticMesh())
-		{
-			ComponentActorString = InStaticMeshComponent->GetStaticMesh()->GetName();
-		}
-	}
-	else if (const USkeletalMeshComponent* InSkeletalMeshComponent = Cast<USkeletalMeshComponent>(InComponent))
-	{
-		if (InSkeletalMeshComponent->SkeletalMesh)
-		{
-			ComponentActorString = InSkeletalMeshComponent->SkeletalMesh->GetName();
-		}
-		else
-		{
-			return NAME_None;
-		}
-	}
-	else
-	{
-		ComponentActorString = InComponent->GetName();
-	}
-
-	if (ComponentActorString.IsEmpty())
-	{
-		return NAME_None;
-	}
-	else
-	{
-		ComponentActorString += TEXT("_") + InComponent->GetOwner()->GetName();
-		return FName(ComponentActorString);
-	}
-}
-
 TSharedPtr<IMeshPaintComponentAdapter> UMeshPaintingSubsystem::GetAdapterForComponent(const UMeshComponent* InComponent) const
 {
-	if (InComponent)
-	{
-		FName ComponentActorName = GetComponentActorName(InComponent);
-		if (ComponentActorName.IsValid() && !ComponentActorName.IsNone())
-		{
-			const TSharedPtr<IMeshPaintComponentAdapter>* MeshAdapterPtr = ComponentToAdapterMap.Find(ComponentActorName);
-			return MeshAdapterPtr ? *MeshAdapterPtr : TSharedPtr<IMeshPaintComponentAdapter>();
-		}
-	}
-
-	return TSharedPtr<IMeshPaintComponentAdapter>();
+	return InComponent ? ComponentToAdapterMap.FindRef(InComponent->GetPathName()) : TSharedPtr<IMeshPaintComponentAdapter>{};
 }
 
 void UMeshPaintingSubsystem::AddToComponentToAdapterMap(const UMeshComponent* InComponent, const TSharedPtr<IMeshPaintComponentAdapter> InAdapter) 
 {
-	if (InComponent && InAdapter)
+	if (!InAdapter || !InComponent)
 	{
-		FName ComponentActorName = GetComponentActorName(InComponent);
-		if (ComponentActorName.IsValid() && !ComponentActorName.IsNone())
-		{
-			ComponentToAdapterMap.Add(ComponentActorName, InAdapter);
-		}
+		return;
 	}
+
+	ComponentToAdapterMap.Add(InComponent->GetPathName(), InAdapter);
 }
 
 TArray<UMeshComponent*> UMeshPaintingSubsystem::GetSelectedMeshComponents() const
diff --git a/Engine/Plugins/MeshPainting/Source/MeshPaintingToolset/Public/MeshPaintHelpers.h b/Engine/Plugins/MeshPainting/Source/MeshPaintingToolset/Public/MeshPaintHelpers.h
--- a/Engine/Plugins/MeshPainting/Source/MeshPaintingToolset/Public/MeshPaintHelpers.h
+++ b/Engine/Plugins/MeshPainting/Source/MeshPaintingToolset/Public/MeshPaintHelpers.h
@@ -275,7 +275,7 @@ private:
 
 private:
 	/** Map of geometry adapters for each selected mesh component */
-	TMap<FName, TSharedPtr<IMeshPaintComponentAdapter>> ComponentToAdapterMap;
+	TMap<FString, TSharedPtr<IMeshPaintComponentAdapter>> ComponentToAdapterMap;
 	TArray<UMeshComponent*> SelectedMeshComponents;
 	/** Mesh components within the current selection which are eligible for painting */
 	TArray<UMeshComponent*> PaintableComponents;
