Fix crash when deleting a mesh that is being painted

diff --git a/Engine/Plugins/MeshPainting/Source/MeshPaintingToolset/Private/MeshPaintSkeletalMeshAdapter.cpp b/Engine/Plugins/MeshPainting/Source/MeshPaintingToolset/Private/MeshPaintSkeletalMeshAdapter.cpp
--- a/Engine/Plugins/MeshPainting/Source/MeshPaintingToolset/Private/MeshPaintSkeletalMeshAdapter.cpp
+++ b/Engine/Plugins/MeshPainting/Source/MeshPaintingToolset/Private/MeshPaintSkeletalMeshAdapter.cpp
@@ -229,11 +229,9 @@ void FMeshPaintSkeletalMeshComponentAdapter::OnAdded()
 
 void FMeshPaintSkeletalMeshComponentAdapter::OnRemoved()
 {
-	checkf(SkeletalMeshComponent, TEXT("Invalid SkeletalMesh Component"));
-	
 	// If the referenced skeletal mesh has been destroyed (and nulled by GC), don't try to do anything more.
 	// It should be in the process of removing all global geometry adapters if it gets here in this situation.
-	if (!ReferencedSkeletalMesh)
+	if (!ReferencedSkeletalMesh || !SkeletalMeshComponent)
 	{
 		return;
 	}
diff --git a/Engine/Plugins/MeshPainting/Source/MeshPaintingToolset/Private/MeshPaintStaticMeshAdapter.cpp b/Engine/Plugins/MeshPainting/Source/MeshPaintingToolset/Private/MeshPaintStaticMeshAdapter.cpp
--- a/Engine/Plugins/MeshPainting/Source/MeshPaintingToolset/Private/MeshPaintStaticMeshAdapter.cpp
+++ b/Engine/Plugins/MeshPainting/Source/MeshPaintingToolset/Private/MeshPaintStaticMeshAdapter.cpp
@@ -174,11 +174,9 @@ void FMeshPaintStaticMeshComponentAdapter::OnAdded()
 
 void FMeshPaintStaticMeshComponentAdapter::OnRemoved()
 {
-	check(StaticMeshComponent);
-	
 	// If the referenced static mesh has been destroyed (and nulled by GC), don't try to do anything more.
 	// It should be in the process of removing all global geometry adapters if it gets here in this situation.
-	if (!ReferencedStaticMesh)
+	if (!ReferencedStaticMesh || !StaticMeshComponent)
 	{
 		return;
 	}
