Various fixes for some mesh paint editor crashes.
- We must not assume that a mesh paint adapter can cache a component pointer and rely on it to be always valid, even if it is added as a referenced object in the FGCObject interface. Components are a type of object which may be force-destroyed (e.g. during an editor tick) and so are best referenced via a TWeakObjectPtr.
- It does not make sense for a RF_Transient component to be mesh painted, as they will never be saved, and hence any edits will be lost (an example of such a component is the dynamically created/destroyed foliage meshes belonging to a landscape actor).  Thus transient components are never added to the adapter map.
- Mesh paint adapters are now keyed on the full component path name, instead of a combination of the owning actor and the mesh name.  Previously it was not possible to edit all components of an actor if they referenced the same mesh.  Also, changing the referenced mesh of a component being edited would result in the adapter no longer being correctly keyed (leading to a crash on lookup).  Meshes generated by the simple construction script of a BP (which are re-instanced with each property change) work, despite their path name changing, because such an event leads to the mesh paint adapters being freshly recached.

diff --git a/Engine/Plugins/MeshPainting/Source/MeshPaintEditorMode/Private/MeshPaintMode.cpp b/Engine/Plugins/MeshPainting/Source/MeshPaintEditorMode/Private/MeshPaintMode.cpp
--- a/Engine/Plugins/MeshPainting/Source/MeshPaintEditorMode/Private/MeshPaintMode.cpp
+++ b/Engine/Plugins/MeshPainting/Source/MeshPaintEditorMode/Private/MeshPaintMode.cpp
@@ -391,11 +391,10 @@ void UMeshPaintMode::UpdateSelectedMeshes()
 		FName PaletteName = Toolkit->GetCurrentPalette();
 
 		MeshPaintingSubsystem->ResetState();
-		const TArray<UMeshComponent*> CurrentMeshComponents = GetSelectedComponents<UMeshComponent>();
+		TArray<UMeshComponent*> CurrentMeshComponents = GetSelectedComponents<UMeshComponent>();
 		MeshPaintingSubsystem->AddSelectedMeshComponents(CurrentMeshComponents);
 		MeshPaintingSubsystem->bNeedsRecache = true;
 
-
 		// Check the current selection for a material that supports texture paint
 		if (PaletteName == UMeshPaintMode::MeshPaintMode_Texture)
 		{
diff --git a/Engine/Plugins/MeshPainting/Source/MeshPaintingToolset/Private/MeshPaintHelpers.cpp b/Engine/Plugins/MeshPainting/Source/MeshPaintingToolset/Private/MeshPaintHelpers.cpp
--- a/Engine/Plugins/MeshPainting/Source/MeshPaintingToolset/Private/MeshPaintHelpers.cpp
+++ b/Engine/Plugins/MeshPainting/Source/MeshPaintingToolset/Private/MeshPaintHelpers.cpp
@@ -25,6 +25,7 @@
 #include "ToolContextInterfaces.h"
 #include "Engine/StaticMesh.h"
 #include "BaseTools/BaseBrushTool.h"
+#include "Algo/Copy.h"
 
 extern void PropagateVertexPaintToSkeletalMesh(USkeletalMesh* SkeletalMesh, int32 LODIndex);
 
@@ -1415,7 +1416,18 @@ void UMeshPaintingSubsystem::AddToComponentToAdapterMap(const UMeshComponent* In
 
 TArray<UMeshComponent*> UMeshPaintingSubsystem::GetSelectedMeshComponents() const
 {
-	return SelectedMeshComponents;
+	TArray<UMeshComponent*> Result;
+	Result.Reserve(SelectedMeshComponents.Num());
+
+	for (TWeakObjectPtr<UMeshComponent> SelectedMeshComponent : SelectedMeshComponents)
+	{
+		if (SelectedMeshComponent.IsValid())
+		{
+			Result.Add(SelectedMeshComponent.Get());
+		}
+	}
+
+	return Result;
 }
 
 void UMeshPaintingSubsystem::ClearSelectedMeshComponents()
@@ -1426,13 +1438,25 @@ void UMeshPaintingSubsystem::ClearSelectedMeshComponents()
 
 void UMeshPaintingSubsystem::AddSelectedMeshComponents(const TArray<UMeshComponent*>& InComponents)
 {
-	SelectedMeshComponents.Append(InComponents);
+	SelectedMeshComponents.Reserve(SelectedMeshComponents.Num() + InComponents.Num());
+	Algo::Copy(InComponents, SelectedMeshComponents);
 }
 
 
 TArray<UMeshComponent*> UMeshPaintingSubsystem::GetPaintableMeshComponents() const
 {
-	return PaintableComponents;
+	TArray<UMeshComponent*> Result;
+	Result.Reserve(PaintableComponents.Num());
+
+	for (TWeakObjectPtr<UMeshComponent> PaintableComponent : PaintableComponents)
+	{
+		if (PaintableComponent.IsValid())
+		{
+			Result.Add(PaintableComponent.Get());
+		}
+	}
+
+	return Result;
 }
 
 void UMeshPaintingSubsystem::AddPaintableMeshComponent(UMeshComponent* InComponent)
@@ -1462,7 +1486,7 @@ void UMeshPaintingSubsystem::Refresh()
 
 void UMeshPaintingSubsystem::CleanUp()
 {
-	for (auto MeshAdapterPair : ComponentToAdapterMap)
+	for (auto& MeshAdapterPair : ComponentToAdapterMap)
 	{
 		MeshAdapterPair.Value->OnRemoved();
 	}
@@ -1480,7 +1504,7 @@ bool UMeshPaintingSubsystem::FindHitResult(const FRay Ray, FHitResult& BestTrace
 		const FVector TraceStart(Origin);
 		const FVector TraceEnd(Origin + Direction * HALF_WORLD_MAX);
 
-		for (UMeshComponent* MeshComponent : PaintableComponents)
+		for (UMeshComponent* MeshComponent : GetPaintableMeshComponents())
 		{
 			TSharedPtr<IMeshPaintComponentAdapter> MeshAdapter = GetAdapterForComponent(MeshComponent);
 			if (!MeshAdapter.IsValid())
@@ -1506,7 +1530,7 @@ bool UMeshPaintingSubsystem::FindHitResult(const FRay Ray, FHitResult& BestTrace
 
 bool UMeshPaintingSubsystem::SelectionContainsValidAdapters() const
 {
-	for (auto& MeshAdapterPair : ComponentToAdapterMap)
+	for (const auto& MeshAdapterPair : ComponentToAdapterMap)
 	{
 		if (MeshAdapterPair.Value && MeshAdapterPair.Value->IsValid())
 		{
@@ -1530,18 +1554,23 @@ void UMeshPaintingSubsystem::SetCopiedColorsByComponent(TArray<FPerComponentVert
 void UMeshPaintingSubsystem::CacheSelectionData(const int32 PaintLODIndex, const int32 UVChannel)
 {
 	bSelectionContainsPerLODColors = false;
-	for (UMeshComponent* MeshComponent : SelectedMeshComponents)
+	for (UMeshComponent* MeshComponent : GetSelectedMeshComponents())
 	{
-		TSharedPtr<IMeshPaintComponentAdapter> MeshAdapter = FMeshPaintComponentAdapterFactory::CreateAdapterForMesh(MeshComponent, PaintLODIndex);
-		if (MeshComponent->IsVisible() && MeshAdapter.IsValid() && MeshAdapter->IsValid())
+		// Don't create an adapter for transient components, as it doesn't make sense to edit them if they will not be saved.
+		// (An example is the procedurally generated foliage meshes)
+		if (!MeshComponent->HasAnyFlags(RF_Transient))
 		{
-			TUniquePtr< FComponentReregisterContext > ComponentReregisterContext;
-			AddPaintableMeshComponent(MeshComponent);
-			AddToComponentToAdapterMap(MeshComponent, MeshAdapter);
-			MeshAdapter->OnAdded();
-			GEngine->GetEngineSubsystem<UMeshPaintingSubsystem>()->ForceRenderMeshLOD(MeshComponent, PaintLODIndex);
-			ComponentReregisterContext.Reset(new FComponentReregisterContext(MeshComponent));
-			bSelectionContainsPerLODColors |= GEngine->GetEngineSubsystem<UMeshPaintingSubsystem>()->DoesMeshComponentContainPerLODColors(MeshComponent);
+			TSharedPtr<IMeshPaintComponentAdapter> MeshAdapter = FMeshPaintComponentAdapterFactory::CreateAdapterForMesh(MeshComponent, PaintLODIndex);
+			if (MeshComponent->IsVisible() && MeshAdapter.IsValid() && MeshAdapter->IsValid())
+			{
+				TUniquePtr< FComponentReregisterContext > ComponentReregisterContext;
+				AddPaintableMeshComponent(MeshComponent);
+				AddToComponentToAdapterMap(MeshComponent, MeshAdapter);
+				MeshAdapter->OnAdded();
+				GEngine->GetEngineSubsystem<UMeshPaintingSubsystem>()->ForceRenderMeshLOD(MeshComponent, PaintLODIndex);
+				ComponentReregisterContext.Reset(new FComponentReregisterContext(MeshComponent));
+				bSelectionContainsPerLODColors |= GEngine->GetEngineSubsystem<UMeshPaintingSubsystem>()->DoesMeshComponentContainPerLODColors(MeshComponent);
+			}
 		}
 	}
 	bNeedsRecache = false;
@@ -1549,9 +1578,9 @@ void UMeshPaintingSubsystem::CacheSelectionData(const int32 PaintLODIndex, const
 
 int32 UMeshPaintingSubsystem::GetMaxUVIndexToPaint() const
 {
-	if (PaintableComponents.Num() == 1 && PaintableComponents[0])
+	if (PaintableComponents.Num() == 1 && PaintableComponents[0].IsValid())
 	{
-		return GEngine->GetEngineSubsystem<UMeshPaintingSubsystem>()->GetNumberOfUVs(PaintableComponents[0], 0) - 1;
+		return GEngine->GetEngineSubsystem<UMeshPaintingSubsystem>()->GetNumberOfUVs(PaintableComponents[0].Get(), 0) - 1;
 	}
 
 	return 0;
diff --git a/Engine/Plugins/MeshPainting/Source/MeshPaintingToolset/Private/MeshPaintSkeletalMeshAdapter.cpp b/Engine/Plugins/MeshPainting/Source/MeshPaintingToolset/Private/MeshPaintSkeletalMeshAdapter.cpp
--- a/Engine/Plugins/MeshPainting/Source/MeshPaintingToolset/Private/MeshPaintSkeletalMeshAdapter.cpp
+++ b/Engine/Plugins/MeshPainting/Source/MeshPaintingToolset/Private/MeshPaintSkeletalMeshAdapter.cpp
@@ -106,7 +106,7 @@ void PropagateVertexPaintToSkeletalMesh(USkeletalMesh* SkeletalMesh, int32 LODIn
 bool FMeshPaintSkeletalMeshComponentAdapter::Construct(UMeshComponent* InComponent, int32 InMeshLODIndex)
 {
 	SkeletalMeshComponent = Cast<USkeletalMeshComponent>(InComponent);
-	if (SkeletalMeshComponent != nullptr)
+	if (SkeletalMeshComponent.IsValid())
 	{
 		SkeletalMeshChangedHandle = SkeletalMeshComponent->RegisterOnSkeletalMeshPropertyChanged(USkeletalMeshComponent::FOnSkeletalMeshPropertyChanged::CreateSP(this, &FMeshPaintSkeletalMeshComponentAdapter::OnSkeletalMeshChanged));
 
@@ -124,7 +124,7 @@ bool FMeshPaintSkeletalMeshComponentAdapter::Construct(UMeshComponent* InCompone
 
 FMeshPaintSkeletalMeshComponentAdapter::~FMeshPaintSkeletalMeshComponentAdapter()
 {
-	if (SkeletalMeshComponent != nullptr)
+	if (SkeletalMeshComponent.IsValid())
 	{
 		SkeletalMeshComponent->UnregisterOnSkeletalMeshPropertyChanged(SkeletalMeshChangedHandle);
 	}
@@ -137,7 +137,7 @@ FMeshPaintSkeletalMeshComponentAdapter::~FMeshPaintSkeletalMeshComponentAdapter(
 void FMeshPaintSkeletalMeshComponentAdapter::OnSkeletalMeshChanged()
 {
 	OnRemoved();
-	if (SkeletalMeshComponent != nullptr)
+	if (SkeletalMeshComponent.IsValid())
 	{
 		ReferencedSkeletalMesh = SkeletalMeshComponent->SkeletalMesh;
 		if (SkeletalMeshComponent->SkeletalMesh != nullptr)
@@ -158,20 +158,23 @@ void FMeshPaintSkeletalMeshComponentAdapter::OnPostMeshCached(USkeletalMesh* Ske
 
 bool FMeshPaintSkeletalMeshComponentAdapter::Initialize()
 {
-	check(ReferencedSkeletalMesh == SkeletalMeshComponent->SkeletalMesh);
-
 	bool bInitializationResult = false;
 
-	MeshResource = ReferencedSkeletalMesh->GetResourceForRendering();
-	if (MeshResource != nullptr)
+	if (SkeletalMeshComponent.IsValid())
 	{
-		LODData = &MeshResource->LODRenderData[MeshLODIndex];
-		checkf(ReferencedSkeletalMesh->GetImportedModel()->LODModels.IsValidIndex(MeshLODIndex), TEXT("Invalid Imported Model index for vertex painting"));
-		LODModel = &ReferencedSkeletalMesh->GetImportedModel()->LODModels[MeshLODIndex];
+		check(ReferencedSkeletalMesh == SkeletalMeshComponent->SkeletalMesh);
 
-		bInitializationResult = FBaseMeshPaintComponentAdapter::Initialize();
-	}
+		MeshResource = ReferencedSkeletalMesh->GetResourceForRendering();
+		if (MeshResource != nullptr)
+		{
+			LODData = &MeshResource->LODRenderData[MeshLODIndex];
+			checkf(ReferencedSkeletalMesh->GetImportedModel()->LODModels.IsValidIndex(MeshLODIndex), TEXT("Invalid Imported Model index for vertex painting"));
+			LODModel = &ReferencedSkeletalMesh->GetImportedModel()->LODModels[MeshLODIndex];
 
+			bInitializationResult = FBaseMeshPaintComponentAdapter::Initialize();
+		}
+
+	}
 	
 	return bInitializationResult;
 }
@@ -210,7 +213,6 @@ void FMeshPaintSkeletalMeshComponentAdapter::AddReferencedObjectsGlobals(FRefere
 void FMeshPaintSkeletalMeshComponentAdapter::AddReferencedObjects(FReferenceCollector& Collector)
 {
 	Collector.AddReferencedObject(ReferencedSkeletalMesh);
-	Collector.AddReferencedObject(SkeletalMeshComponent);
 }
 
 void FMeshPaintSkeletalMeshComponentAdapter::CleanupGlobals()
@@ -219,7 +221,13 @@ void FMeshPaintSkeletalMeshComponentAdapter::CleanupGlobals()
 
 void FMeshPaintSkeletalMeshComponentAdapter::OnAdded()
 {
-	checkf(SkeletalMeshComponent, TEXT("Invalid SkeletalMesh Component"));
+	// We shouldn't assume that the cached skeletal mesh component remains valid.
+	// Components may be destroyed by editor ticks, and be forcibly removed by GC.
+	if (!SkeletalMeshComponent.IsValid())
+	{
+		return;
+	}
+
 	checkf(ReferencedSkeletalMesh, TEXT("Invalid reference to Skeletal Mesh"));
 	checkf(ReferencedSkeletalMesh == SkeletalMeshComponent->SkeletalMesh, TEXT("Referenced Skeletal Mesh does not match one in Component"));
 
@@ -232,9 +240,16 @@ void FMeshPaintSkeletalMeshComponentAdapter::OnAdded()
 
 void FMeshPaintSkeletalMeshComponentAdapter::OnRemoved()
 {
+	// We shouldn't assume that the cached skeletal mesh component remains valid.
+	// Components may be destroyed by editor ticks, and be forcibly removed by GC.
+	if (!SkeletalMeshComponent.IsValid())
+	{
+		return;
+	}
+	
 	// If the referenced skeletal mesh has been destroyed (and nulled by GC), don't try to do anything more.
 	// It should be in the process of removing all global geometry adapters if it gets here in this situation.
-	if (!ReferencedSkeletalMesh || !SkeletalMeshComponent)
+	if (!ReferencedSkeletalMesh)
 	{
 		return;
 	}
@@ -247,6 +262,11 @@ void FMeshPaintSkeletalMeshComponentAdapter::OnRemoved()
 
 bool FMeshPaintSkeletalMeshComponentAdapter::LineTraceComponent(struct FHitResult& OutHit, const FVector Start, const FVector End, const struct FCollisionQueryParams& Params) const
 {
+	if (!SkeletalMeshComponent.IsValid())
+	{
+		return false;
+	}
+
 	const bool bHitBounds = FMath::LineSphereIntersection(Start, End.GetSafeNormal(), (End - Start).SizeSquared(), SkeletalMeshComponent->Bounds.Origin, SkeletalMeshComponent->Bounds.SphereRadius);
 	const float SqrRadius = FMath::Square(SkeletalMeshComponent->Bounds.SphereRadius);
 	const bool bInsideBounds = (SkeletalMeshComponent->Bounds.ComputeSquaredDistanceFromBoxToPoint(Start) <= SqrRadius) || (SkeletalMeshComponent->Bounds.ComputeSquaredDistanceFromBoxToPoint(End) <= SqrRadius);
@@ -313,12 +333,18 @@ bool FMeshPaintSkeletalMeshComponentAdapter::LineTraceComponent(struct FHitResul
 
 void FMeshPaintSkeletalMeshComponentAdapter::QueryPaintableTextures(int32 MaterialIndex, int32& OutDefaultIndex, TArray<struct FPaintableTexture>& InOutTextureList)
 {
-	DefaultQueryPaintableTextures(MaterialIndex, SkeletalMeshComponent, OutDefaultIndex, InOutTextureList);
+	if (SkeletalMeshComponent.IsValid())
+	{
+		DefaultQueryPaintableTextures(MaterialIndex, SkeletalMeshComponent.Get(), OutDefaultIndex, InOutTextureList);
+	}
 }
 
 void FMeshPaintSkeletalMeshComponentAdapter::ApplyOrRemoveTextureOverride(UTexture* SourceTexture, UTexture* OverrideTexture) const
 {
-	DefaultApplyOrRemoveTextureOverride(SkeletalMeshComponent, SourceTexture, OverrideTexture);
+	if (SkeletalMeshComponent.IsValid())
+	{
+		DefaultApplyOrRemoveTextureOverride(SkeletalMeshComponent.Get(), SourceTexture, OverrideTexture);
+	}
 }
 
 void FMeshPaintSkeletalMeshComponentAdapter::GetTextureCoordinate(int32 VertexIndex, int32 ChannelIndex, FVector2D& OutTextureCoordinate) const
@@ -328,6 +354,11 @@ void FMeshPaintSkeletalMeshComponentAdapter::GetTextureCoordinate(int32 VertexIn
 
 void FMeshPaintSkeletalMeshComponentAdapter::PreEdit()
 {
+	if (!SkeletalMeshComponent.IsValid())
+	{
+		return;
+	}
+
 	FlushRenderingCommands();
 
 	SkeletalMeshComponent->Modify();
@@ -419,6 +450,11 @@ void FMeshPaintSkeletalMeshComponentAdapter::SetVertexColor(int32 VertexIndex, F
 
 FMatrix FMeshPaintSkeletalMeshComponentAdapter::GetComponentToWorldMatrix() const
 {
+	if (!SkeletalMeshComponent.IsValid())
+	{
+		return FMatrix::Identity;
+	}
+
 	return SkeletalMeshComponent->GetComponentToWorld().ToMatrixWithScale();
 }
 
diff --git a/Engine/Plugins/MeshPainting/Source/MeshPaintingToolset/Private/MeshPaintSplineMeshAdapter.cpp b/Engine/Plugins/MeshPainting/Source/MeshPaintingToolset/Private/MeshPaintSplineMeshAdapter.cpp
--- a/Engine/Plugins/MeshPainting/Source/MeshPaintingToolset/Private/MeshPaintSplineMeshAdapter.cpp
+++ b/Engine/Plugins/MeshPainting/Source/MeshPaintingToolset/Private/MeshPaintSplineMeshAdapter.cpp
@@ -10,8 +10,13 @@
 
 bool FMeshPaintSplineMeshComponentAdapter::InitializeVertexData()
 {
+	if (!StaticMeshComponent.IsValid())
+	{
+		return false;
+	}
+
 	// Cache deformed spline mesh vertices for quick lookup during painting / previewing
-	USplineMeshComponent* SplineMeshComponent = Cast<USplineMeshComponent>(StaticMeshComponent);
+	USplineMeshComponent* SplineMeshComponent = Cast<USplineMeshComponent>(StaticMeshComponent.Get());
 	check(SplineMeshComponent);
 
 	bool bValid = false;
diff --git a/Engine/Plugins/MeshPainting/Source/MeshPaintingToolset/Private/MeshPaintStaticMeshAdapter.cpp b/Engine/Plugins/MeshPainting/Source/MeshPaintingToolset/Private/MeshPaintStaticMeshAdapter.cpp
--- a/Engine/Plugins/MeshPainting/Source/MeshPaintingToolset/Private/MeshPaintStaticMeshAdapter.cpp
+++ b/Engine/Plugins/MeshPainting/Source/MeshPaintingToolset/Private/MeshPaintStaticMeshAdapter.cpp
@@ -14,7 +14,7 @@
 bool FMeshPaintStaticMeshComponentAdapter::Construct(UMeshComponent* InComponent, int32 InMeshLODIndex)
 {
 	StaticMeshComponent = Cast<UStaticMeshComponent>(InComponent);
-	if (StaticMeshComponent != nullptr)
+	if (StaticMeshComponent.IsValid())
 	{
 #if WITH_EDITOR
 		StaticMeshComponent->OnStaticMeshChanged().AddRaw(this, &FMeshPaintStaticMeshComponentAdapter::OnStaticMeshChanged);
@@ -37,7 +37,7 @@ bool FMeshPaintStaticMeshComponentAdapter::Construct(UMeshComponent* InComponent
 FMeshPaintStaticMeshComponentAdapter::~FMeshPaintStaticMeshComponentAdapter()
 {
 #if WITH_EDITOR
-	if (StaticMeshComponent != nullptr)
+	if (StaticMeshComponent.IsValid())
 	{
 		StaticMeshComponent->OnStaticMeshChanged().RemoveAll(this);
 	}
@@ -64,7 +64,6 @@ void FMeshPaintStaticMeshComponentAdapter::OnStaticMeshChanged(UStaticMeshCompon
 		ReferencedStaticMesh->OnPostMeshBuild().RemoveAll(this);
 #endif
 		ReferencedStaticMesh = InStaticMeshComponent->GetStaticMesh();
-
 #if WITH_EDITOR
 		ReferencedStaticMesh->OnPostMeshBuild().AddRaw(this, &FMeshPaintStaticMeshComponentAdapter::OnPostMeshBuild);
 #endif
@@ -75,11 +74,14 @@ void FMeshPaintStaticMeshComponentAdapter::OnStaticMeshChanged(UStaticMeshCompon
 
 bool FMeshPaintStaticMeshComponentAdapter::Initialize()
 {
-	check(ReferencedStaticMesh == StaticMeshComponent->GetStaticMesh());
-	if (MeshLODIndex < ReferencedStaticMesh->GetNumLODs())
+	if (StaticMeshComponent.IsValid())
 	{
-		LODModel = &(ReferencedStaticMesh->GetRenderData()->LODResources[MeshLODIndex]);
-		return FBaseMeshPaintComponentAdapter::Initialize();
+		check(ReferencedStaticMesh == StaticMeshComponent->GetStaticMesh());
+		if (MeshLODIndex < ReferencedStaticMesh->GetNumLODs())
+		{
+			LODModel = &(ReferencedStaticMesh->GetRenderData()->LODResources[MeshLODIndex]);
+			return FBaseMeshPaintComponentAdapter::Initialize();
+		}
 	}
 
 	return false;
@@ -111,6 +113,13 @@ bool FMeshPaintStaticMeshComponentAdapter::InitializeVertexData()
 
 void FMeshPaintStaticMeshComponentAdapter::PostEdit()
 {
+	// We shouldn't assume that the cached static mesh component remains valid.
+	// Components may be destroyed by editor ticks, and be forcibly removed by GC.
+	if (!StaticMeshComponent.IsValid())
+	{
+		return;
+	}
+
 	// Lighting does not need to be invalidated when mesh painting
 	const bool bUnbuildLighting = false;
 
@@ -123,7 +132,7 @@ void FMeshPaintStaticMeshComponentAdapter::PostEdit()
 	{
 		// We're only changing instanced vertices on this specific mesh component, so we
 		// only need to detach our mesh component
-		FComponentReregisterContext ComponentReregisterContext(StaticMeshComponent);
+		FComponentReregisterContext ComponentReregisterContext(StaticMeshComponent.Get());
 
 		// If LOD is 0, post-edit all LODs. There's currently no way to tell from here
 		// if VertexPaintSettings.bPaintOnSpecificLOD is set to true or not.
@@ -161,29 +170,16 @@ void FMeshPaintStaticMeshComponentAdapter::CleanupGlobals()
 void FMeshPaintStaticMeshComponentAdapter::AddReferencedObjects(FReferenceCollector& Collector)
 {
 	Collector.AddReferencedObject(ReferencedStaticMesh);
-	Collector.AddReferencedObject(StaticMeshComponent);
 }
 
 
-void FMeshPaintStaticMeshComponentAdapter::OnAdded()
-{
-	check(StaticMeshComponent);
-	check(ReferencedStaticMesh);
-	check(ReferencedStaticMesh == StaticMeshComponent->GetStaticMesh());
-}
-
-void FMeshPaintStaticMeshComponentAdapter::OnRemoved()
+bool FMeshPaintStaticMeshComponentAdapter::LineTraceComponent(struct FHitResult& OutHit, const FVector Start, const FVector End, const struct FCollisionQueryParams& Params) const
 {
-	// If the referenced static mesh has been destroyed (and nulled by GC), don't try to do anything more.
-	// It should be in the process of removing all global geometry adapters if it gets here in this situation.
-	if (!ReferencedStaticMesh || !StaticMeshComponent)
+	if (!StaticMeshComponent.IsValid())
 	{
-		return;
+		return false;
 	}
-}
 
-bool FMeshPaintStaticMeshComponentAdapter::LineTraceComponent(struct FHitResult& OutHit, const FVector Start, const FVector End, const struct FCollisionQueryParams& Params) const
-{
 	// Ray trace
 	const bool bHitBounds = FMath::LineSphereIntersection(Start, End.GetSafeNormal(), (End - Start).SizeSquared(), StaticMeshComponent->Bounds.Origin, StaticMeshComponent->Bounds.SphereRadius);
 	const float SqrRadius = FMath::Square(StaticMeshComponent->Bounds.SphereRadius);
@@ -249,18 +245,30 @@ bool FMeshPaintStaticMeshComponentAdapter::LineTraceComponent(struct FHitResult&
 
 void FMeshPaintStaticMeshComponentAdapter::QueryPaintableTextures(int32 MaterialIndex, int32& OutDefaultIndex, TArray<struct FPaintableTexture>& InOutTextureList)
 {
-	DefaultQueryPaintableTextures(MaterialIndex, StaticMeshComponent, OutDefaultIndex, InOutTextureList);
+	if (StaticMeshComponent.IsValid())
+	{
+		DefaultQueryPaintableTextures(MaterialIndex, StaticMeshComponent.Get(), OutDefaultIndex, InOutTextureList);
+	}
 }
 
 void FMeshPaintStaticMeshComponentAdapter::ApplyOrRemoveTextureOverride(UTexture* SourceTexture, UTexture* OverrideTexture) const
 {
-	DefaultApplyOrRemoveTextureOverride(StaticMeshComponent, SourceTexture, OverrideTexture);
+	if (StaticMeshComponent.IsValid())
+	{
+		DefaultApplyOrRemoveTextureOverride(StaticMeshComponent.Get(), SourceTexture, OverrideTexture);
+	}
 }
 
 void FMeshPaintStaticMeshComponentAdapter::GetVertexColor(int32 VertexIndex, FColor& OutColor, bool bInstance /*= true*/) const
 {
 	if (bInstance)
 	{
+		if (!StaticMeshComponent.IsValid())
+		{
+			OutColor = FColor::White;
+			return;
+		}
+
 		FStaticMeshComponentLODInfo* InstanceMeshLODInfo = &StaticMeshComponent->LODData[MeshLODIndex];
 		if (!bInstance && LODModel->VertexBuffers.ColorVertexBuffer.GetNumVertices() == 0)
 		{
@@ -296,6 +304,11 @@ void FMeshPaintStaticMeshComponentAdapter::SetVertexColor(int32 VertexIndex, FCo
 	// Update the mesh!				
 	if (bInstance)
 	{
+		if (!StaticMeshComponent.IsValid())
+		{
+			return;
+		}
+
 		FStaticMeshComponentLODInfo* InstanceMeshLODInfo = &StaticMeshComponent->LODData[MeshLODIndex];
 		const bool bValidInstanceData = InstanceMeshLODInfo
 			&& InstanceMeshLODInfo->OverrideVertexColors
@@ -331,6 +344,11 @@ void FMeshPaintStaticMeshComponentAdapter::SetVertexColor(int32 VertexIndex, FCo
 
 FMatrix FMeshPaintStaticMeshComponentAdapter::GetComponentToWorldMatrix() const
 {
+	if (!StaticMeshComponent.IsValid())
+	{
+		return FMatrix::Identity;
+	}
+
 	return StaticMeshComponent->GetComponentToWorld().ToMatrixWithScale();
 }
 
@@ -342,6 +360,11 @@ void FMeshPaintStaticMeshComponentAdapter::GetTextureCoordinate(int32 VertexInde
 
 void FMeshPaintStaticMeshComponentAdapter::PreEdit()
 {
+	if (!StaticMeshComponent.IsValid())
+	{
+		return;
+	}
+
 	const bool bUsingInstancedVertexColors = true; // Currently we are only painting to instances 
 	UStaticMesh* StaticMesh = ReferencedStaticMesh;
 	if (bUsingInstancedVertexColors)
diff --git a/Engine/Plugins/MeshPainting/Source/MeshPaintingToolset/Public/MeshPaintHelpers.h b/Engine/Plugins/MeshPainting/Source/MeshPaintingToolset/Public/MeshPaintHelpers.h
--- a/Engine/Plugins/MeshPainting/Source/MeshPaintingToolset/Public/MeshPaintHelpers.h
+++ b/Engine/Plugins/MeshPainting/Source/MeshPaintingToolset/Public/MeshPaintHelpers.h
@@ -269,16 +269,18 @@ protected:
 	bool bSelectionHasMaterialValidForTexturePaint;
 
 private:
-	FName GetComponentActorName(const UMeshComponent* InComponent) const;
-
 	void CleanUp();
 
 private:
 	/** Map of geometry adapters for each selected mesh component */
 	TMap<FString, TSharedPtr<IMeshPaintComponentAdapter>> ComponentToAdapterMap;
-	TArray<UMeshComponent*> SelectedMeshComponents;
+
+	/** Currently selected mesh components as provided by the mode class */
+	TArray<TWeakObjectPtr<UMeshComponent>> SelectedMeshComponents;
+
 	/** Mesh components within the current selection which are eligible for painting */
-	TArray<UMeshComponent*> PaintableComponents;
+	TArray<TWeakObjectPtr<UMeshComponent>> PaintableComponents;
+
 	/** Contains copied vertex color data */
 	TArray<FPerComponentVertexColorData> CopiedColorsByComponent;
 	bool bSelectionContainsPerLODColors;
diff --git a/Engine/Plugins/MeshPainting/Source/MeshPaintingToolset/Public/MeshPaintSkeletalMeshAdapter.h b/Engine/Plugins/MeshPainting/Source/MeshPaintingToolset/Public/MeshPaintSkeletalMeshAdapter.h
--- a/Engine/Plugins/MeshPainting/Source/MeshPaintingToolset/Public/MeshPaintSkeletalMeshAdapter.h
+++ b/Engine/Plugins/MeshPainting/Source/MeshPaintingToolset/Public/MeshPaintSkeletalMeshAdapter.h
@@ -31,9 +31,9 @@ public:
 	virtual bool Initialize() override;
 	virtual void OnAdded() override;
 	virtual void OnRemoved() override;
-	virtual bool IsValid() const override { return SkeletalMeshComponent && ReferencedSkeletalMesh && SkeletalMeshComponent->SkeletalMesh == ReferencedSkeletalMesh; }
+	virtual bool IsValid() const override { return SkeletalMeshComponent.IsValid() && ReferencedSkeletalMesh && SkeletalMeshComponent->SkeletalMesh == ReferencedSkeletalMesh; }
 	virtual bool SupportsTexturePaint() const override { return true; }
-	virtual bool SupportsVertexPaint() const override { return SkeletalMeshComponent != nullptr; }
+	virtual bool SupportsVertexPaint() const override { return SkeletalMeshComponent.IsValid(); }
 	virtual bool LineTraceComponent(struct FHitResult& OutHit, const FVector Start, const FVector End, const struct FCollisionQueryParams& Params) const override;
 	virtual void QueryPaintableTextures(int32 MaterialIndex, int32& OutDefaultIndex, TArray<struct FPaintableTexture>& InOutTextureList) override;
 	virtual void ApplyOrRemoveTextureOverride(UTexture* SourceTexture, UTexture* OverrideTexture) const override;
@@ -59,7 +59,7 @@ protected:
 	FDelegateHandle SkeletalMeshChangedHandle;	
 
 	/** Skeletal mesh component represented by this adapter */
-	USkeletalMeshComponent* SkeletalMeshComponent;
+	TWeakObjectPtr<USkeletalMeshComponent> SkeletalMeshComponent;
 	/** Skeletal mesh currently set to the Skeletal Mesh Component */
 	USkeletalMesh* ReferencedSkeletalMesh;
 	/** Skeletal Mesh resource retrieved from the Skeletal Mesh */
diff --git a/Engine/Plugins/MeshPainting/Source/MeshPaintingToolset/Public/MeshPaintStaticMeshAdapter.h b/Engine/Plugins/MeshPainting/Source/MeshPaintingToolset/Public/MeshPaintStaticMeshAdapter.h
--- a/Engine/Plugins/MeshPainting/Source/MeshPaintingToolset/Public/MeshPaintStaticMeshAdapter.h
+++ b/Engine/Plugins/MeshPainting/Source/MeshPaintingToolset/Public/MeshPaintStaticMeshAdapter.h
@@ -26,11 +26,11 @@ public:
 	virtual bool Construct(UMeshComponent* InComponent, int32 InMeshLODIndex) override;
 	virtual ~FMeshPaintStaticMeshComponentAdapter();
 	virtual bool Initialize() override;
-	virtual void OnAdded() override;
-	virtual void OnRemoved() override;
-	virtual bool IsValid() const override { return StaticMeshComponent && StaticMeshComponent->GetStaticMesh() == ReferencedStaticMesh; }
+	virtual void OnAdded() override {}
+	virtual void OnRemoved() override {}
+	virtual bool IsValid() const override { return StaticMeshComponent.IsValid() && ReferencedStaticMesh && StaticMeshComponent->GetStaticMesh() == ReferencedStaticMesh; }
 	virtual bool SupportsTexturePaint() const override { return true; }
-	virtual bool SupportsVertexPaint() const override { return StaticMeshComponent && !StaticMeshComponent->bDisallowMeshPaintPerInstance; }
+	virtual bool SupportsVertexPaint() const override { return StaticMeshComponent.IsValid() && !StaticMeshComponent->bDisallowMeshPaintPerInstance; }
 	virtual bool LineTraceComponent(struct FHitResult& OutHit, const FVector Start, const FVector End, const struct FCollisionQueryParams& Params) const override;	
 	virtual void QueryPaintableTextures(int32 MaterialIndex, int32& OutDefaultIndex, TArray<struct FPaintableTexture>& InOutTextureList) override;
 	virtual void ApplyOrRemoveTextureOverride(UTexture* SourceTexture, UTexture* OverrideTexture) const override;
@@ -54,7 +54,7 @@ protected:
 	void OnStaticMeshChanged(UStaticMeshComponent* StaticMeshComponent);
 
 	/** Static mesh component represented by this adapter */
-	UStaticMeshComponent* StaticMeshComponent;
+	TWeakObjectPtr<UStaticMeshComponent> StaticMeshComponent;
 	/** Static mesh currently set to the Static Mesh Component */
 	UStaticMesh* ReferencedStaticMesh;
 	/** LOD model (at Mesh LOD Index) containing data to change */
