Fixed a crash while using the mesh painting tool with an unsuported texture.
The code for the tool seems to only support BRGA8 textures.

diff --git a/Engine/Plugins/MeshPainting/Source/MeshPaintingToolset/Private/TexturePaintToolset.cpp b/Engine/Plugins/MeshPainting/Source/MeshPaintingToolset/Private/TexturePaintToolset.cpp
--- a/Engine/Plugins/MeshPainting/Source/MeshPaintingToolset/Private/TexturePaintToolset.cpp
+++ b/Engine/Plugins/MeshPainting/Source/MeshPaintingToolset/Private/TexturePaintToolset.cpp
@@ -496,6 +496,18 @@ void UTexturePaintToolset::RetrieveTexturesForComponent(const UMeshComponent* Co
 			int32 OutDefaultIndex = 0;
 			Adapter->QueryPaintableTextures(MaterialIndex, OutDefaultIndex, OutTextures);
 		}
+
+		// Filter out any paintable texture that is not supported currently (Todo: should we do this here or could we just move some logic in the FPaintableTexture)?
+		OutTextures.RemoveAll([](const FPaintableTexture& PaintableTexture)
+			{
+				if (const UTexture* Texture = PaintableTexture.Texture)
+				{
+					// Only BRGA8 image are supported by the tool currently (we should probably provide some feedback to the user on why some texture are not supported and a way to convert them)
+					return Texture->Source.GetFormat() != TSF_BGRA8;
+				}
+
+				return true;
+			});
 	}
 }
 
