Renamed CopyTextureWithTransitions to TransitionAndCopyTexture.

diff --git a/Engine/Plugins/Compositing/OpenCVLensDistortion/Source/OpenCVLensDistortion/Private/LensDistortionDisplacementMapRendering.cpp b/Engine/Plugins/Compositing/OpenCVLensDistortion/Source/OpenCVLensDistortion/Private/LensDistortionDisplacementMapRendering.cpp
--- a/Engine/Plugins/Compositing/OpenCVLensDistortion/Source/OpenCVLensDistortion/Private/LensDistortionDisplacementMapRendering.cpp
+++ b/Engine/Plugins/Compositing/OpenCVLensDistortion/Source/OpenCVLensDistortion/Private/LensDistortionDisplacementMapRendering.cpp
@@ -153,7 +153,7 @@ static void DrawUVDisplacementToRenderTargetFromPreComputedDisplacementMap_Rende
 		RHICmdList.DrawPrimitive(0, PrimitiveCount, 1);
 	}
 	RHICmdList.EndRenderPass();
-	CopyTextureWithTransitions(RHICmdList, OutTextureRenderTargetResource->GetRenderTargetTexture(), OutTextureRenderTargetResource->TextureRHI, {});
+	TransitionAndCopyTexture(RHICmdList, OutTextureRenderTargetResource->GetRenderTargetTexture(), OutTextureRenderTargetResource->TextureRHI, {});
 #endif
 }
 
diff --git a/Engine/Plugins/FX/Niagara/Source/Niagara/Private/NiagaraDataInterfaceGrid2DCollection.cpp b/Engine/Plugins/FX/Niagara/Source/Niagara/Private/NiagaraDataInterfaceGrid2DCollection.cpp
--- a/Engine/Plugins/FX/Niagara/Source/Niagara/Private/NiagaraDataInterfaceGrid2DCollection.cpp
+++ b/Engine/Plugins/FX/Niagara/Source/Niagara/Private/NiagaraDataInterfaceGrid2DCollection.cpp
@@ -2050,25 +2050,6 @@ void UNiagaraDataInterfaceGrid2DCollection::FindAttributes(TArray<FNiagaraVariab
 	OutNumAttribChannelsFound = TotalAttributes - NumAttributes;
 }
 
-static void TransitionAndCopyTexture(FRHICommandList& RHICmdList, FRHITexture* Source, FRHITexture* Destination, const FRHICopyTextureInfo& CopyInfo)
-{
-	FRHITransitionInfo TransitionsBefore[] = {
-		FRHITransitionInfo(Source, ERHIAccess::SRVMask, ERHIAccess::CopySrc),
-		FRHITransitionInfo(Destination, ERHIAccess::SRVMask, ERHIAccess::CopyDest)
-	};
-
-	RHICmdList.Transition(MakeArrayView(TransitionsBefore, UE_ARRAY_COUNT(TransitionsBefore)));
-
-	RHICmdList.CopyTexture(Source, Destination, CopyInfo);
-
-	FRHITransitionInfo TransitionsAfter[] = {
-		FRHITransitionInfo(Source, ERHIAccess::CopySrc, ERHIAccess::SRVMask),
-		FRHITransitionInfo(Destination, ERHIAccess::CopyDest, ERHIAccess::SRVMask)
-	};
-
-	RHICmdList.Transition(MakeArrayView(TransitionsAfter, UE_ARRAY_COUNT(TransitionsAfter)));
-}
-
 void UNiagaraDataInterfaceGrid2DCollection::GetCanvasVariables(TArray<FNiagaraVariableBase>& OutVariables) const
 {
 	TArray<uint32> VariableOffsets;
diff --git a/Engine/Plugins/MeshPainting/Source/MeshPaintingToolset/Private/MeshTexturePaintingTool.cpp b/Engine/Plugins/MeshPainting/Source/MeshPaintingToolset/Private/MeshTexturePaintingTool.cpp
--- a/Engine/Plugins/MeshPainting/Source/MeshPaintingToolset/Private/MeshTexturePaintingTool.cpp
+++ b/Engine/Plugins/MeshPainting/Source/MeshPaintingToolset/Private/MeshTexturePaintingTool.cpp
@@ -1001,7 +1001,7 @@ void UMeshTexturePaintingTool::PaintTexture(FMeshPaintParameters& InParams, TArr
 	ENQUEUE_RENDER_COMMAND(UpdateMeshPaintRTCommand1)(
 		[BrushRenderTargetResource](FRHICommandListImmediate& RHICmdList)
 	{
-		CopyTextureWithTransitions(RHICmdList, BrushRenderTargetResource->GetRenderTargetTexture(), BrushRenderTargetResource->TextureRHI, {});
+		TransitionAndCopyTexture(RHICmdList, BrushRenderTargetResource->GetRenderTargetTexture(), BrushRenderTargetResource->TextureRHI, {});
 	});
 
 	if (bEnableSeamPainting)
@@ -1011,7 +1011,7 @@ void UMeshTexturePaintingTool::PaintTexture(FMeshPaintParameters& InParams, TArr
 		ENQUEUE_RENDER_COMMAND(UpdateMeshPaintRTCommand2)(
 			[BrushMaskRenderTargetResource](FRHICommandListImmediate& RHICmdList)
 		{
-			CopyTextureWithTransitions(RHICmdList, BrushMaskRenderTargetResource->GetRenderTargetTexture(), BrushMaskRenderTargetResource->TextureRHI, {});
+			TransitionAndCopyTexture(RHICmdList, BrushMaskRenderTargetResource->GetRenderTargetTexture(), BrushMaskRenderTargetResource->TextureRHI, {});
 		});
 	}
 
@@ -1105,7 +1105,7 @@ void UMeshTexturePaintingTool::PaintTexture(FMeshPaintParameters& InParams, TArr
 		ENQUEUE_RENDER_COMMAND(UpdateMeshPaintRTCommand3)(
 			[RenderTargetResource](FRHICommandListImmediate& RHICmdList)
 		{
-			CopyTextureWithTransitions(RHICmdList, RenderTargetResource->GetRenderTargetTexture(), RenderTargetResource->TextureRHI, {});
+			TransitionAndCopyTexture(RHICmdList, RenderTargetResource->GetRenderTargetTexture(), RenderTargetResource->TextureRHI, {});
 		});
 	}
 	FlushRenderingCommands();
diff --git a/Engine/Plugins/MeshPainting/Source/MeshPaintingToolset/Private/TexturePaintToolset.cpp b/Engine/Plugins/MeshPainting/Source/MeshPaintingToolset/Private/TexturePaintToolset.cpp
--- a/Engine/Plugins/MeshPainting/Source/MeshPaintingToolset/Private/TexturePaintToolset.cpp
+++ b/Engine/Plugins/MeshPainting/Source/MeshPaintingToolset/Private/TexturePaintToolset.cpp
@@ -96,7 +96,7 @@ void UTexturePaintToolset::CopyTextureToRenderTargetTexture(UTexture* SourceText
 	ENQUEUE_RENDER_COMMAND(UpdateMeshPaintRTCommand)(
 		[RenderTargetResource](FRHICommandListImmediate& RHICmdList)
 		{
-			CopyTextureWithTransitions(RHICmdList, RenderTargetResource->GetRenderTargetTexture(), RenderTargetResource->TextureRHI, {});
+			TransitionAndCopyTexture(RHICmdList, RenderTargetResource->GetRenderTargetTexture(), RenderTargetResource->TextureRHI, {});
 		});
 }
 
@@ -305,7 +305,7 @@ bool UTexturePaintToolset::GenerateSeamMask(UMeshComponent* MeshComponent, int32
 		ENQUEUE_RENDER_COMMAND(UpdateMeshPaintRTCommand5)(
 			[RenderTargetResource](FRHICommandListImmediate& RHICmdList)
 		{
-			CopyTextureWithTransitions(RHICmdList, RenderTargetResource->GetRenderTargetTexture(), RenderTargetResource->TextureRHI, {});
+			TransitionAndCopyTexture(RHICmdList, RenderTargetResource->GetRenderTargetTexture(), RenderTargetResource->TextureRHI, {});
 		});
 	}
 
diff --git a/Engine/Plugins/MovieScene/MovieRenderPipeline/Source/MovieRenderPipelineCore/Private/MoviePipelineSurfaceReader.cpp b/Engine/Plugins/MovieScene/MovieRenderPipeline/Source/MovieRenderPipelineCore/Private/MoviePipelineSurfaceReader.cpp
--- a/Engine/Plugins/MovieScene/MovieRenderPipeline/Source/MovieRenderPipelineCore/Private/MoviePipelineSurfaceReader.cpp
+++ b/Engine/Plugins/MovieScene/MovieRenderPipeline/Source/MovieRenderPipelineCore/Private/MoviePipelineSurfaceReader.cpp
@@ -162,7 +162,7 @@ void FMoviePipelineSurfaceReader::ResolveSampleToReadbackTexture_RenderThread(co
 	}
 	RHICmdList.EndRenderPass();
 
-	CopyTextureWithTransitions(RHICmdList, ResampleTexturePooledRenderTarget->GetRHI(), ReadbackTexture, {});
+	TransitionAndCopyTexture(RHICmdList, ResampleTexturePooledRenderTarget->GetRHI(), ReadbackTexture, {});
 }
 
 void FMoviePipelineSurfaceReader::CopyReadbackTexture_RenderThread(TUniqueFunction<void(TUniquePtr<FImagePixelData>&&)>&& InFunctionCallback, TSharedPtr<FImagePixelDataPayload, ESPMode::ThreadSafe> InFramePayload)
diff --git a/Engine/Plugins/Runtime/TextureShare/Source/TextureShare/Private/TextureShareRHI.cpp b/Engine/Plugins/Runtime/TextureShare/Source/TextureShare/Private/TextureShareRHI.cpp
--- a/Engine/Plugins/Runtime/TextureShare/Source/TextureShare/Private/TextureShareRHI.cpp
+++ b/Engine/Plugins/Runtime/TextureShare/Source/TextureShare/Private/TextureShareRHI.cpp
@@ -71,7 +71,7 @@ namespace
 		Params.Size.X = SrcRect.Width();
 		Params.Size.Y = SrcRect.Height();
 
-		CopyTextureWithTransitions(RHICmdList, SrcTexture, DstTexture, {});
+		TransitionAndCopyTexture(RHICmdList, SrcTexture, DstTexture, {});
 	}
 
 	static void ResampleCopyTextureImpl_RenderThread(FRHICommandListImmediate& RHICmdList, FRHITexture* SrcTexture, FRHITexture* DstTexture, const FIntRect* SrcTextureRect, const FIntRect* DstTextureRect)
diff --git a/Engine/Plugins/Runtime/TextureShare/Source/TextureShareD3D12/Private/CrossGPUHeap/D3D12CrossGPUHeap.cpp b/Engine/Plugins/Runtime/TextureShare/Source/TextureShareD3D12/Private/CrossGPUHeap/D3D12CrossGPUHeap.cpp
--- a/Engine/Plugins/Runtime/TextureShare/Source/TextureShareD3D12/Private/CrossGPUHeap/D3D12CrossGPUHeap.cpp
+++ b/Engine/Plugins/Runtime/TextureShare/Source/TextureShareD3D12/Private/CrossGPUHeap/D3D12CrossGPUHeap.cpp
@@ -205,7 +205,7 @@ static void CopyDirectTextureImpl_RenderThread(FRHICommandListImmediate& RHICmdL
 		Params.DestPosition.Y = DstRect.Min.Y;
 	}
 
-	CopyTextureWithTransitions(RHICmdList, SrcTexture, DstTexture, {});
+	TransitionAndCopyTexture(RHICmdList, SrcTexture, DstTexture, {});
 }
 
 bool FD3D12CrossGPUHeap::SendCrossGPUResource(FRHICommandListImmediate& RHICmdList, const FString& ResourceID, FRHITexture2D* SrcResource, const FIntRect* SrcTextureRect)
diff --git a/Engine/Plugins/Runtime/nDisplay/Source/DisplayCluster/Private/Render/Device/DisplayClusterDeviceBase.cpp b/Engine/Plugins/Runtime/nDisplay/Source/DisplayCluster/Private/Render/Device/DisplayClusterDeviceBase.cpp
--- a/Engine/Plugins/Runtime/nDisplay/Source/DisplayCluster/Private/Render/Device/DisplayClusterDeviceBase.cpp
+++ b/Engine/Plugins/Runtime/nDisplay/Source/DisplayCluster/Private/Render/Device/DisplayClusterDeviceBase.cpp
@@ -480,7 +480,7 @@ void FDisplayClusterDeviceBase::RenderTexture_RenderThread(FRHICommandListImmedi
 			CopyInfo.Size.X = FMath::Min(SrcSize.X, DstSize.X);
 			CopyInfo.Size.Y = FMath::Min(SrcSize.Y, DstSize.Y);
 
-			CopyTextureWithTransitions(RHICmdList, SrcTexture, BackBuffer, CopyInfo);
+			TransitionAndCopyTexture(RHICmdList, SrcTexture, BackBuffer, CopyInfo);
 		}
 
 		if (RenderFrameMode == EDisplayClusterRenderFrameMode::Stereo && ViewportManagerProxyPtr)
diff --git a/Engine/Plugins/Runtime/nDisplay/Source/DisplayCluster/Private/Render/Viewport/DisplayClusterViewportManagerProxy.cpp b/Engine/Plugins/Runtime/nDisplay/Source/DisplayCluster/Private/Render/Viewport/DisplayClusterViewportManagerProxy.cpp
--- a/Engine/Plugins/Runtime/nDisplay/Source/DisplayCluster/Private/Render/Viewport/DisplayClusterViewportManagerProxy.cpp
+++ b/Engine/Plugins/Runtime/nDisplay/Source/DisplayCluster/Private/Render/Viewport/DisplayClusterViewportManagerProxy.cpp
@@ -558,7 +558,7 @@ bool FDisplayClusterViewportManagerProxy::ResolveFrameTargetToBackBuffer_RenderT
 			CopyInfo.DestPosition.X = DstRect.Min.X;
 			CopyInfo.DestPosition.Y = DstRect.Min.Y;
 
-			CopyTextureWithTransitions(RHICmdList, FrameTexture, DestTexture, CopyInfo);
+			TransitionAndCopyTexture(RHICmdList, FrameTexture, DestTexture, CopyInfo);
 
 			return true;
 		}
diff --git a/Engine/Plugins/Runtime/nDisplay/Source/DisplayCluster/Private/Render/Viewport/DisplayClusterViewportProxy.cpp b/Engine/Plugins/Runtime/nDisplay/Source/DisplayCluster/Private/Render/Viewport/DisplayClusterViewportProxy.cpp
--- a/Engine/Plugins/Runtime/nDisplay/Source/DisplayCluster/Private/Render/Viewport/DisplayClusterViewportProxy.cpp
+++ b/Engine/Plugins/Runtime/nDisplay/Source/DisplayCluster/Private/Render/Viewport/DisplayClusterViewportProxy.cpp
@@ -592,7 +592,7 @@ void ImplResolveResource(FRHICommandListImmediate& RHICmdList, FRHITexture2D* In
 		CopyInfo.DestPosition.X = OutputRect.Min.X;
 		CopyInfo.DestPosition.Y = OutputRect.Min.Y;
 
-		CopyTextureWithTransitions(RHICmdList, InputResource, OutputResource, CopyInfo);
+		TransitionAndCopyTexture(RHICmdList, InputResource, OutputResource, CopyInfo);
 	}
 	else
 	{
diff --git a/Engine/Plugins/Runtime/nDisplay/Source/DisplayCluster/Private/Render/Viewport/Postprocess/DisplayClusterViewportPostProcessOutputRemap.cpp b/Engine/Plugins/Runtime/nDisplay/Source/DisplayCluster/Private/Render/Viewport/Postprocess/DisplayClusterViewportPostProcessOutputRemap.cpp
--- a/Engine/Plugins/Runtime/nDisplay/Source/DisplayCluster/Private/Render/Viewport/Postprocess/DisplayClusterViewportPostProcessOutputRemap.cpp
+++ b/Engine/Plugins/Runtime/nDisplay/Source/DisplayCluster/Private/Render/Viewport/Postprocess/DisplayClusterViewportPostProcessOutputRemap.cpp
@@ -270,7 +270,7 @@ void FDisplayClusterViewportPostProcessOutputRemap::PerformPostProcessFrame_Rend
 
 				if (ShadersAPI.RenderPostprocess_OutputRemap(RHICmdList, InOutTexture, TempTargetableTexture, *MeshProxy))
 				{
-					CopyTextureWithTransitions(RHICmdList, TempTargetableTexture, InOutTexture, {});
+					TransitionAndCopyTexture(RHICmdList, TempTargetableTexture, InOutTexture, {});
 				}
 			}
 		}
diff --git a/Engine/Plugins/VirtualProduction/DMX/DMXPixelMapping/Source/DMXPixelMappingRenderer/Private/DMXPixelMappingRenderer.cpp b/Engine/Plugins/VirtualProduction/DMX/DMXPixelMapping/Source/DMXPixelMappingRenderer/Private/DMXPixelMappingRenderer.cpp
--- a/Engine/Plugins/VirtualProduction/DMX/DMXPixelMapping/Source/DMXPixelMappingRenderer/Private/DMXPixelMappingRenderer.cpp
+++ b/Engine/Plugins/VirtualProduction/DMX/DMXPixelMapping/Source/DMXPixelMappingRenderer/Private/DMXPixelMappingRenderer.cpp
@@ -214,7 +214,7 @@ void FDMXPixelMappingRenderer::DownsampleRender(
 
 			// Copy texture from GPU to CPU
 			{
-				CopyTextureWithTransitions(RHICmdList, ResolveRenderTarget, RenderTargetRef, {});
+				TransitionAndCopyTexture(RHICmdList, ResolveRenderTarget, RenderTargetRef, {});
 
 				// Read the contents of a texture to an output CPU buffer
 				TArray<FLinearColor> ColorArray;
diff --git a/Engine/Source/Editor/MeshPaint/Private/TexturePaintHelpers.cpp b/Engine/Source/Editor/MeshPaint/Private/TexturePaintHelpers.cpp
--- a/Engine/Source/Editor/MeshPaint/Private/TexturePaintHelpers.cpp
+++ b/Engine/Source/Editor/MeshPaint/Private/TexturePaintHelpers.cpp
@@ -96,7 +96,7 @@ void TexturePaintHelpers::CopyTextureToRenderTargetTexture(UTexture* SourceTextu
 	ENQUEUE_RENDER_COMMAND(UpdateMeshPaintRTCommand)(
 		[RenderTargetResource](FRHICommandListImmediate& RHICmdList)
 		{
-			CopyTextureWithTransitions(RHICmdList, RenderTargetResource->GetRenderTargetTexture(), RenderTargetResource->TextureRHI, {});
+			TransitionAndCopyTexture(RHICmdList, RenderTargetResource->GetRenderTargetTexture(), RenderTargetResource->TextureRHI, {});
 		});		
 }
 
@@ -305,7 +305,7 @@ bool TexturePaintHelpers::GenerateSeamMask(UMeshComponent* MeshComponent, int32
 		ENQUEUE_RENDER_COMMAND(UpdateMeshPaintRTCommand5)(
 			[RenderTargetResource](FRHICommandListImmediate& RHICmdList)
 			{
-				CopyTextureWithTransitions(RHICmdList, RenderTargetResource->GetRenderTargetTexture(), RenderTargetResource->TextureRHI, {});
+				TransitionAndCopyTexture(RHICmdList, RenderTargetResource->GetRenderTargetTexture(), RenderTargetResource->TextureRHI, {});
 			});
 	}
 
diff --git a/Engine/Source/Editor/UnrealEd/Private/ObjectTools.cpp b/Engine/Source/Editor/UnrealEd/Private/ObjectTools.cpp
--- a/Engine/Source/Editor/UnrealEd/Private/ObjectTools.cpp
+++ b/Engine/Source/Editor/UnrealEd/Private/ObjectTools.cpp
@@ -4790,7 +4790,7 @@ namespace ThumbnailTools
 				ENQUEUE_RENDER_COMMAND(UpdateThumbnailRTCommand)(
 					[RenderTargetResource](FRHICommandListImmediate& RHICmdList)
 					{
-						CopyTextureWithTransitions(RHICmdList, RenderTargetResource->GetRenderTargetTexture(), RenderTargetResource->TextureRHI, {});
+						TransitionAndCopyTexture(RHICmdList, RenderTargetResource->GetRenderTargetTexture(), RenderTargetResource->TextureRHI, {});
 					});
 
 				if(OutThumbnail)
diff --git a/Engine/Source/Runtime/Engine/Private/KismetRenderingLibrary.cpp b/Engine/Source/Runtime/Engine/Private/KismetRenderingLibrary.cpp
--- a/Engine/Source/Runtime/Engine/Private/KismetRenderingLibrary.cpp
+++ b/Engine/Source/Runtime/Engine/Private/KismetRenderingLibrary.cpp
@@ -781,7 +781,7 @@ void UKismetRenderingLibrary::EndDrawCanvasToRenderTarget(UObject* WorldContextO
 			ENQUEUE_RENDER_COMMAND(CanvasRenderTargetResolveCommand)(
 				[RenderTargetResource](FRHICommandListImmediate& RHICmdList)
 				{
-					CopyTextureWithTransitions(RHICmdList, RenderTargetResource->GetRenderTargetTexture(), RenderTargetResource->TextureRHI, {});
+					TransitionAndCopyTexture(RHICmdList, RenderTargetResource->GetRenderTargetTexture(), RenderTargetResource->TextureRHI, {});
 				}
 			);
 
diff --git a/Engine/Source/Runtime/Engine/Private/UnrealClient.cpp b/Engine/Source/Runtime/Engine/Private/UnrealClient.cpp
--- a/Engine/Source/Runtime/Engine/Private/UnrealClient.cpp
+++ b/Engine/Source/Runtime/Engine/Private/UnrealClient.cpp
@@ -1715,7 +1715,7 @@ const TArray<FColor>& FViewport::GetRawHitProxyData(FIntRect InRect)
 		ENQUEUE_RENDER_COMMAND(UpdateHitProxyRTCommand)(
 			[HitProxyMapPtr](FRHICommandListImmediate& RHICmdList)
 			{
-				CopyTextureWithTransitions(RHICmdList, HitProxyMapPtr->GetRenderTargetTexture(), HitProxyMapPtr->GetHitProxyCPUTexture(), {});
+				TransitionAndCopyTexture(RHICmdList, HitProxyMapPtr->GetRenderTargetTexture(), HitProxyMapPtr->GetHitProxyCPUTexture(), {});
 			});
 
 		ENQUEUE_RENDER_COMMAND(EndDrawingCommand)(
diff --git a/Engine/Source/Runtime/GameplayMediaEncoder/Private/GameplayMediaEncoder.cpp b/Engine/Source/Runtime/GameplayMediaEncoder/Private/GameplayMediaEncoder.cpp
--- a/Engine/Source/Runtime/GameplayMediaEncoder/Private/GameplayMediaEncoder.cpp
+++ b/Engine/Source/Runtime/GameplayMediaEncoder/Private/GameplayMediaEncoder.cpp
@@ -689,7 +689,7 @@ void FGameplayMediaEncoder::CopyTexture(const FTexture2DRHIRef& SourceTexture, F
 
 	if(SourceTexture->GetFormat() == DestinationTexture->GetFormat() && SourceTexture->GetSizeXY() == DestinationTexture->GetSizeXY())
 	{
-		CopyTextureWithTransitions(RHICmdList, SourceTexture, DestinationTexture, {});
+		TransitionAndCopyTexture(RHICmdList, SourceTexture, DestinationTexture, {});
 	}
 	else // Texture format mismatch, use a shader to do the copy.
 	{
diff --git a/Engine/Source/Runtime/Landscape/Private/LandscapeEdit.cpp b/Engine/Source/Runtime/Landscape/Private/LandscapeEdit.cpp
--- a/Engine/Source/Runtime/Landscape/Private/LandscapeEdit.cpp
+++ b/Engine/Source/Runtime/Landscape/Private/LandscapeEdit.cpp
@@ -7541,7 +7541,7 @@ bool ALandscapeProxy::LandscapeExportHeightmapToRenderTarget(UTextureRenderTarge
 	ENQUEUE_RENDER_COMMAND(DrawHeightmapRTCommand)(
 		[RenderTargetResource](FRHICommandListImmediate& RHICmdList)
 		{
-			CopyTextureWithTransitions(RHICmdList, RenderTargetResource->GetRenderTargetTexture(), RenderTargetResource->TextureRHI, {});
+			TransitionAndCopyTexture(RHICmdList, RenderTargetResource->GetRenderTargetTexture(), RenderTargetResource->TextureRHI, {});
 		});
 
 
diff --git a/Engine/Source/Runtime/Landscape/Private/LandscapeEditLayers.cpp b/Engine/Source/Runtime/Landscape/Private/LandscapeEditLayers.cpp
--- a/Engine/Source/Runtime/Landscape/Private/LandscapeEditLayers.cpp
+++ b/Engine/Source/Runtime/Landscape/Private/LandscapeEditLayers.cpp
@@ -3269,7 +3269,7 @@ void ALandscape::PrintLayersDebugRT(const FString& InContext, UTextureRenderTarg
 		[RenderTargetResource](FRHICommandListImmediate& RHICmdList) mutable
 	{
 		TRACE_CPUPROFILER_EVENT_SCOPE(LandscapeLayers_RT_DebugResolve);
-		CopyTextureWithTransitions(RHICmdList, RenderTargetResource->GetRenderTargetTexture(), RenderTargetResource->TextureRHI, {});
+		TransitionAndCopyTexture(RHICmdList, RenderTargetResource->GetRenderTargetTexture(), RenderTargetResource->TextureRHI, {});
 	});
 
 	{
diff --git a/Engine/Source/Runtime/MovieSceneCapture/Private/FrameGrabber.cpp b/Engine/Source/Runtime/MovieSceneCapture/Private/FrameGrabber.cpp
--- a/Engine/Source/Runtime/MovieSceneCapture/Private/FrameGrabber.cpp
+++ b/Engine/Source/Runtime/MovieSceneCapture/Private/FrameGrabber.cpp
@@ -174,7 +174,7 @@ void FViewportSurfaceReader::ResolveRenderTarget(FViewportSurfaceReader* RenderT
 				EDRF_Default);
 		}
 		RHICmdList.EndRenderPass();
-		CopyTextureWithTransitions(RHICmdList, DestRenderTarget, ReadbackTexture, {});
+		TransitionAndCopyTexture(RHICmdList, DestRenderTarget, ReadbackTexture, {});
 
 		if (RenderToReadback)
 		{
diff --git a/Engine/Source/Runtime/RHI/Public/RHIUtilities.h b/Engine/Source/Runtime/RHI/Public/RHIUtilities.h
--- a/Engine/Source/Runtime/RHI/Public/RHIUtilities.h
+++ b/Engine/Source/Runtime/RHI/Public/RHIUtilities.h
@@ -799,7 +799,7 @@ inline void ClearRenderTarget(FRHICommandList& RHICmdList, FRHITexture* Texture,
 	RHICmdList.EndRenderPass();
 }
 
-inline void CopyTextureWithTransitions(FRHICommandListImmediate& RHICmdList, FRHITexture* SrcTexture, FRHITexture* DstTexture, const FRHICopyTextureInfo& Info)
+inline void TransitionAndCopyTexture(FRHICommandList& RHICmdList, FRHITexture* SrcTexture, FRHITexture* DstTexture, const FRHICopyTextureInfo& Info)
 {
 	check(SrcTexture && DstTexture);
 	check(SrcTexture->GetNumSamples() == DstTexture->GetNumSamples());
diff --git a/Engine/Source/Runtime/SlateRHIRenderer/Private/SlateRHIRenderer.cpp b/Engine/Source/Runtime/SlateRHIRenderer/Private/SlateRHIRenderer.cpp
--- a/Engine/Source/Runtime/SlateRHIRenderer/Private/SlateRHIRenderer.cpp
+++ b/Engine/Source/Runtime/SlateRHIRenderer/Private/SlateRHIRenderer.cpp
@@ -930,7 +930,7 @@ void FSlateRHIRenderer::DrawWindow_RenderThread(FRHICommandListImmediate& RHICmd
 				bClear = true; // Force a clear of the UI buffer to black
 
 				// Grab HDR backbuffer
-				CopyTextureWithTransitions(RHICmdList, FinalBuffer, ViewportInfo.HDRSourceRT->GetRHI(), {});
+				TransitionAndCopyTexture(RHICmdList, FinalBuffer, ViewportInfo.HDRSourceRT->GetRHI(), {});
 
 				// UI backbuffer is temp target
 				BackBuffer = ViewportInfo.UITargetRT->GetRHI();
