Fix the MeshTexturePaintingTool OnPaintingFinished() to return a reference of the delegate instead of a copy.

diff --git a/Engine/Plugins/MeshPainting/Source/MeshPaintingToolset/Public/MeshTexturePaintingTool.h b/Engine/Plugins/MeshPainting/Source/MeshPaintingToolset/Public/MeshTexturePaintingTool.h
--- a/Engine/Plugins/MeshPainting/Source/MeshPaintingToolset/Public/MeshTexturePaintingTool.h
+++ b/Engine/Plugins/MeshPainting/Source/MeshPaintingToolset/Public/MeshTexturePaintingTool.h
@@ -124,7 +124,7 @@ public:
 	}
 
 
-	FSimpleDelegate OnPaintingFinished()
+	FSimpleDelegate& OnPaintingFinished()
 	{
 		return OnPaintingFinishedDelegate;
 	}
