Fixed a crash while using the mesh painting tools in a SM6 feature level editor.

diff --git a/Engine/Plugins/MeshPainting/Source/MeshPaintingToolset/Private/MeshTexturePaintingTool.cpp b/Engine/Plugins/MeshPainting/Source/MeshPaintingToolset/Private/MeshTexturePaintingTool.cpp
--- a/Engine/Plugins/MeshPainting/Source/MeshPaintingToolset/Private/MeshTexturePaintingTool.cpp
+++ b/Engine/Plugins/MeshPainting/Source/MeshPaintingToolset/Private/MeshTexturePaintingTool.cpp
@@ -20,6 +20,7 @@
 #include "BaseGizmos/BrushStampIndicator.h"
 #include "MeshPaintAdapterFactory.h"
 #include "TextureCompiler.h"
+#include "RHI.h"
 
 #define LOCTEXT_NAMESPACE "MeshTextureBrush"
 
@@ -805,7 +806,7 @@ void UMeshTexturePaintingTool::PaintTexture(FMeshPaintParameters& InParams, TArr
 		return;
 	}
 
-	const auto FeatureLevel = ERHIFeatureLevel::SM5;// GEditor->GetEditorWorldContext().World()->FeatureLevel;
+	const auto FeatureLevel = GMaxRHIFeatureLevel;
 
 
 	FPaintTexture2DData* TextureData = GetPaintTargetData(PaintingTexture2D);
@@ -1262,7 +1263,7 @@ void UMeshTexturePaintingTool::CommitAllPaintedTextures()
 
 void UMeshTexturePaintingTool::ClearAllTextureOverrides()
 {
-	const auto FeatureLevel = ERHIFeatureLevel::SM5;// GEditor->GetEditorWorldContext().World()->FeatureLevel;
+	const auto FeatureLevel = GMaxRHIFeatureLevel;
 	/** Remove all texture overrides which are currently stored and active */
 	for (decltype(PaintTargetData)::TIterator It(PaintTargetData); It; ++It)
 	{
