Add support for Geometry Collection in Mesh Painting mode
 + Fix bounding box calculation issues in GeometryCollectionComponent (which caused issues w/ painting)
 + Add vertex color rendering to Geometry Collection Scene Proxy (so you can see what you're painting)

diff --git a/Engine/Plugins/Experimental/ChaosEditor/Source/FractureEditor/Private/FractureEditorModeToolkit.cpp b/Engine/Plugins/Experimental/ChaosEditor/Source/FractureEditor/Private/FractureEditorModeToolkit.cpp
--- a/Engine/Plugins/Experimental/ChaosEditor/Source/FractureEditor/Private/FractureEditorModeToolkit.cpp
+++ b/Engine/Plugins/Experimental/ChaosEditor/Source/FractureEditor/Private/FractureEditorModeToolkit.cpp
@@ -1784,6 +1784,7 @@ void FFractureEditorModeToolkit::UpdateExplodedVectors(UGeometryCollectionCompon
 	}
 
 	GeometryCollectionComponent->RefreshEmbeddedGeometry();
+	GeometryCollectionComponent->UpdateCachedBounds();
 }
 
 void FFractureEditorModeToolkit::RegenerateOutliner()
diff --git a/Engine/Plugins/MeshPainting/Source/MeshPaintEditorMode/MeshPaintEditorMode.Build.cs b/Engine/Plugins/MeshPainting/Source/MeshPaintEditorMode/MeshPaintEditorMode.Build.cs
--- a/Engine/Plugins/MeshPainting/Source/MeshPaintEditorMode/MeshPaintEditorMode.Build.cs
+++ b/Engine/Plugins/MeshPainting/Source/MeshPaintEditorMode/MeshPaintEditorMode.Build.cs
@@ -27,6 +27,7 @@ namespace UnrealBuildTool.Rules
 					"InteractiveToolsFramework",
 					"EditorInteractiveToolsFramework",
 					"MeshPaintingToolset",
+					"GeometryCollectionEngine",
 					"PropertyEditor",
 					"MainFrame",
 					"DesktopPlatform",
diff --git a/Engine/Plugins/MeshPainting/Source/MeshPaintEditorMode/Private/BrushSettingsCustomization.cpp b/Engine/Plugins/MeshPainting/Source/MeshPaintEditorMode/Private/BrushSettingsCustomization.cpp
--- a/Engine/Plugins/MeshPainting/Source/MeshPaintEditorMode/Private/BrushSettingsCustomization.cpp
+++ b/Engine/Plugins/MeshPainting/Source/MeshPaintEditorMode/Private/BrushSettingsCustomization.cpp
@@ -29,6 +29,7 @@
 #include "MeshPaintModeHelpers.h"
 #include "MeshTexturePaintingTool.h"
 #include "ContentBrowserDelegates.h"
+#include "GeometryCollection/GeometryCollectionComponent.h"
 
 #define LOCTEXT_NAMESPACE "MeshPaintCustomization"
 
@@ -303,10 +304,17 @@ void FColorPaintingSettingsCustomization::CustomizeDetails(IDetailLayoutBuilder&
 			{
 				static const FText SkelMeshNotificationText = LOCTEXT("SkelMeshAssetPaintInfo", "Paint is propagated to Skeletal Mesh Asset(s)");
 				static const FText StaticMeshNotificationText = LOCTEXT("StaticMeshAssetPaintInfo", "Paint is applied to all LODs");
+				static const FText GeometryCollectionNotificationText = LOCTEXT("GeometryCollectionAssetPaintInfo", "Paint is propagated to Geometry Collection Asset(s), and Geometry Collection does not currently support LODs.");
 
+				const bool bGeometryCollectionText = UMeshPaintMode::GetMeshPaintMode()->GetSelectedComponents<UGeometryCollectionComponent>().Num() > 0;
 				const bool bSkelMeshText = UMeshPaintMode::GetMeshPaintMode()->GetSelectedComponents<USkeletalMeshComponent>().Num() > 0;
 				const bool bLODPaintText = UMeshPaintMode::GetColorToolProperties() ?  !UMeshPaintMode::GetColorToolProperties()->bPaintOnSpecificLOD : false;
-				return FText::Format(FTextFormat::FromString(TEXT("{0}{1}{2}")), bSkelMeshText ? SkelMeshNotificationText : FText::GetEmpty(), bSkelMeshText && bLODPaintText ? FText::FromString(TEXT("\n")) : FText::GetEmpty(), bLODPaintText ? StaticMeshNotificationText : FText::GetEmpty());
+				return FText::Format(FTextFormat::FromString(TEXT("{0}{1}{2}{3}")), 
+					bSkelMeshText ? SkelMeshNotificationText : FText::GetEmpty(),
+					bGeometryCollectionText ? GeometryCollectionNotificationText : FText::GetEmpty(),
+					bSkelMeshText && bLODPaintText ? FText::FromString(TEXT("\n")) : FText::GetEmpty(),
+					bLODPaintText ? StaticMeshNotificationText : FText::GetEmpty()
+					);
 			})
 		];
 }
diff --git a/Engine/Plugins/MeshPainting/Source/MeshPaintEditorMode/Private/MeshPaintEditorMode.cpp b/Engine/Plugins/MeshPainting/Source/MeshPaintEditorMode/Private/MeshPaintEditorMode.cpp
--- a/Engine/Plugins/MeshPainting/Source/MeshPaintEditorMode/Private/MeshPaintEditorMode.cpp
+++ b/Engine/Plugins/MeshPainting/Source/MeshPaintEditorMode/Private/MeshPaintEditorMode.cpp
@@ -13,6 +13,7 @@
 #include "MeshPaintSplineMeshAdapter.h"
 #include "MeshPaintStaticMeshAdapter.h"
 #include "MeshPaintSkeletalMeshAdapter.h"
+#include "MeshPaintGeometryCollectionAdapter.h"
 
 class FMeshPaintEditorModeModule : public IModuleInterface
 {
@@ -49,6 +50,7 @@ void FMeshPaintEditorModeModule::Register()
 	RegisterGeometryAdapterFactory(MakeShareable(new FMeshPaintSplineMeshComponentAdapterFactory));
 	RegisterGeometryAdapterFactory(MakeShareable(new FMeshPaintStaticMeshComponentAdapterFactory));
 	RegisterGeometryAdapterFactory(MakeShareable(new FMeshPaintSkeletalMeshComponentAdapterFactory));
+	RegisterGeometryAdapterFactory(MakeShareable(new FMeshPaintGeometryCollectionComponentAdapterFactory));
 
 	FPropertyEditorModule& PropertyModule = FModuleManager::LoadModuleChecked<FPropertyEditorModule>("PropertyEditor");
 	PropertyModule.RegisterCustomClassLayout("MeshColorPaintingToolProperties", FOnGetDetailCustomizationInstance::CreateStatic(&FColorPaintingSettingsCustomization::MakeInstance));
diff --git a/Engine/Plugins/MeshPainting/Source/MeshPaintEditorMode/Private/MeshPaintMode.cpp b/Engine/Plugins/MeshPainting/Source/MeshPaintEditorMode/Private/MeshPaintMode.cpp
--- a/Engine/Plugins/MeshPainting/Source/MeshPaintEditorMode/Private/MeshPaintMode.cpp
+++ b/Engine/Plugins/MeshPainting/Source/MeshPaintEditorMode/Private/MeshPaintMode.cpp
@@ -15,6 +15,7 @@
 #include "Components/StaticMeshComponent.h"
 #include "MeshPaintHelpers.h"
 #include "Components/SkeletalMeshComponent.h"
+#include "GeometryCollection/GeometryCollectionComponent.h"
 #include "ScopedTransaction.h"
 #include "PackageTools.h"
 #include "MeshPaintHelpers.h"
@@ -827,6 +828,7 @@ TArray<ComponentClass*> UMeshPaintMode::GetSelectedComponents() const
 template TArray<UStaticMeshComponent*> UMeshPaintMode::GetSelectedComponents<UStaticMeshComponent>() const;
 template TArray<USkeletalMeshComponent*> UMeshPaintMode::GetSelectedComponents<USkeletalMeshComponent>() const;
 template TArray<UMeshComponent*> UMeshPaintMode::GetSelectedComponents<UMeshComponent>() const;
+template TArray<UGeometryCollectionComponent*> UMeshPaintMode::GetSelectedComponents<UGeometryCollectionComponent>() const;
 
 
 void UMeshPaintMode::UpdateCachedVertexDataSize()
diff --git a/Engine/Plugins/MeshPainting/Source/MeshPaintingToolset/MeshPaintingToolset.Build.cs b/Engine/Plugins/MeshPainting/Source/MeshPaintingToolset/MeshPaintingToolset.Build.cs
--- a/Engine/Plugins/MeshPainting/Source/MeshPaintingToolset/MeshPaintingToolset.Build.cs
+++ b/Engine/Plugins/MeshPainting/Source/MeshPaintingToolset/MeshPaintingToolset.Build.cs
@@ -21,6 +21,7 @@ namespace UnrealBuildTool.Rules
 					"SlateCore",
 					"MeshDescription",
 					"StaticMeshDescription",
+					"GeometryCollectionEngine",
 					"TypedElementFramework",
 					"TypedElementRuntime",
 					"UnrealEd",
diff --git a/Engine/Plugins/MeshPainting/Source/MeshPaintingToolset/Private/MeshPaintGeometryCollectionAdapter.cpp b/Engine/Plugins/MeshPainting/Source/MeshPaintingToolset/Private/MeshPaintGeometryCollectionAdapter.cpp
new file mode 100644
--- /dev/null
+++ b/Engine/Plugins/MeshPainting/Source/MeshPaintingToolset/Private/MeshPaintGeometryCollectionAdapter.cpp
@@ -0,0 +1,314 @@
+// Copyright Epic Games, Inc. All Rights Reserved.
+
+#include "MeshPaintGeometryCollectionAdapter.h"
+#include "GeometryCollection/GeometryCollectionComponent.h"
+#include "MeshPaintHelpers.h"
+#include "MeshPaintingToolsetTypes.h"
+#include "ComponentReregisterContext.h"
+#include "IndexTypes.h"
+
+DEFINE_LOG_CATEGORY_STATIC(LogMeshPaintGeometryCollectionComponentAdapter, Log, All);
+//////////////////////////////////////////////////////////////////////////
+// FMeshPaintGeometryCollectionComponentAdapter
+
+
+
+bool FMeshPaintGeometryCollectionComponentAdapter::Construct(UMeshComponent* InComponent, int32 InMeshLODIndex)
+{
+	GeometryCollectionComponent = Cast<UGeometryCollectionComponent>(InComponent);
+	if (HasValidGeometryCollection())
+	{
+		GeometryCollectionChangedHandle = GeometryCollectionComponent->RegisterOnGeometryCollectionPropertyChanged(UGeometryCollectionComponent::FOnGeometryCollectionPropertyChanged::CreateSP(this, &FMeshPaintGeometryCollectionComponentAdapter::OnGeometryCollectionChanged));
+
+		// MeshLODIndex = InMeshLODIndex; // TODO: Store the LOD index to work on, once GeometryCollection supports LODs
+		const bool bSuccess = Initialize();
+		return bSuccess;
+	}
+
+	return false;
+}
+
+FMeshPaintGeometryCollectionComponentAdapter::~FMeshPaintGeometryCollectionComponentAdapter()
+{
+	if (GeometryCollectionComponent.IsValid())
+	{
+		GeometryCollectionComponent->UnregisterOnGeometryCollectionPropertyChanged(GeometryCollectionChangedHandle);
+	}
+}
+
+UGeometryCollection* FMeshPaintGeometryCollectionComponentAdapter::GetGeometryCollectionObject() const
+{
+	checkSlow(GeometryCollectionComponent.IsValid()); // expect this to be called only when the component is already verified as valid
+	// Note: We bypass the FGeometryCollectionEdit machinery here because it doesn't fit with the model of the Mesh Paint Component Adapter
+	// Instead, we will manually call Modify() as needed
+	return const_cast<UGeometryCollection*>(GeometryCollectionComponent->GetRestCollection());
+}
+
+void FMeshPaintGeometryCollectionComponentAdapter::OnGeometryCollectionChanged()
+{
+	// Remove and re-initialize/add the geometry collection to make sure it is up to date after something changed
+	OnRemoved();
+	if (HasValidGeometryCollection())
+	{
+		Initialize();
+		OnAdded();
+	}
+}
+
+bool FMeshPaintGeometryCollectionComponentAdapter::Initialize()
+{
+	bool bInitializationResult = false;
+
+	if (HasValidGeometryCollection())
+	{
+		bInitializationResult = FBaseMeshPaintComponentAdapter::Initialize(); // Calls InitializeVertexData() then builds octree, etc
+	}
+	
+	return bInitializationResult;
+}
+
+bool FMeshPaintGeometryCollectionComponentAdapter::HasValidGeometryCollection() const
+{
+	return GeometryCollectionComponent.IsValid() && GetGeometryCollectionObject() && GetGeometryCollectionObject()->GetGeometryCollection();
+}
+
+bool FMeshPaintGeometryCollectionComponentAdapter::IsValid() const
+{
+	return GeometryCollectionComponent.IsValid() && GetGeometryCollectionObject() && GetGeometryCollectionObject()->GetGeometryCollection() // Pointers valid
+		&& GetGeometryCollectionObject()->GetGeometryCollection()->Vertex.Num() == MeshVertices.Num(); // Data matches what we've converted
+}
+
+bool FMeshPaintGeometryCollectionComponentAdapter::InitializeVertexData()
+{
+	// Retrieve mesh vertex and index data 
+	FGeometryCollection* Collection = GetGeometryCollectionObject()->GetGeometryCollection().Get(); 
+	const TManagedArray<FVector3f>& Vertices = Collection->Vertex;
+	// TODO: Revisit this when we support the same geometry mapping to multiple bones
+	const TManagedArray<int32>& BoneMap = Collection->BoneMap;
+	const TManagedArray<int32>& SimTypes = Collection->SimulationType;
+	const TManagedArray<bool>& Visible = Collection->Visible;
+	const TArray<FMatrix>& Matrices = GeometryCollectionComponent->GetGlobalMatrices();
+	const int32 NumVertices = Vertices.Num();
+	MeshVertices.Reset();
+	MeshVertices.AddZeroed(NumVertices);
+	for (int32 Index = 0; Index < NumVertices; ++Index)
+	{
+		int32 BoneIndex = BoneMap[Index];
+		const FMatrix& Matrix = Matrices[BoneIndex];
+		const FVector Position = Matrix.TransformPosition((FVector)Vertices[Index]);
+		MeshVertices[Index] = Position;
+	}
+
+	const TManagedArray<FIntVector>& ManagedIndices = Collection->Indices;
+	int32 NumTris = ManagedIndices.Num();
+	for (int32 TriIdx = 0; TriIdx < NumTris; ++TriIdx)
+	{
+		if (Visible[TriIdx])
+		{
+			MeshIndices.Add(ManagedIndices[TriIdx].X);
+			MeshIndices.Add(ManagedIndices[TriIdx].Y);
+			MeshIndices.Add(ManagedIndices[TriIdx].Z);
+		}
+	}
+
+	return (MeshVertices.Num() >= 0 && MeshIndices.Num() > 0);
+}
+
+void FMeshPaintGeometryCollectionComponentAdapter::InitializeAdapterGlobals()
+{
+	static bool bInitialized = false;
+	if (!bInitialized)
+	{
+		bInitialized = true;
+	}
+}
+
+void FMeshPaintGeometryCollectionComponentAdapter::AddReferencedObjectsGlobals(FReferenceCollector& Collector)
+{
+}
+
+void FMeshPaintGeometryCollectionComponentAdapter::AddReferencedObjects(FReferenceCollector& Collector)
+{
+}
+
+void FMeshPaintGeometryCollectionComponentAdapter::CleanupGlobals()
+{
+}
+
+void FMeshPaintGeometryCollectionComponentAdapter::OnAdded()
+{
+	// We shouldn't assume that the cached geometry collection component remains valid.
+	// Components may be destroyed by editor ticks, and be forcibly removed by GC.
+	if (!GeometryCollectionComponent.IsValid())
+	{
+		return;
+	}
+
+	checkf(GetGeometryCollectionObject(), TEXT("Geometry Collection Component did not have a valid geometry collection object attached"));
+	checkf(GetGeometryCollectionObject()->GetGeometryCollection(), TEXT("Geometry Collection Component did not have valid geometry collection data attached"));
+}
+
+void FMeshPaintGeometryCollectionComponentAdapter::OnRemoved()
+{
+}
+
+bool FMeshPaintGeometryCollectionComponentAdapter::LineTraceComponent(struct FHitResult& OutHit, const FVector Start, const FVector End, const struct FCollisionQueryParams& Params) const
+{
+	// Note: This is 99% the same as the skeletal mesh component adapter version, and could likely be pushed up to the base class if we added a generalized way to get the Component
+
+	if (!IsValid())
+	{
+		return false;
+	}
+
+	const bool bHitBounds = FMath::LineSphereIntersection(Start, End.GetSafeNormal(), (End - Start).SizeSquared(), GeometryCollectionComponent->Bounds.Origin, GeometryCollectionComponent->Bounds.SphereRadius);
+	const float SqrRadius = FMath::Square(GeometryCollectionComponent->Bounds.SphereRadius);
+	const bool bInsideBounds = (GeometryCollectionComponent->Bounds.ComputeSquaredDistanceFromBoxToPoint(Start) <= SqrRadius) || (GeometryCollectionComponent->Bounds.ComputeSquaredDistanceFromBoxToPoint(End) <= SqrRadius);
+
+	bool bHitTriangle = false;
+	if (bHitBounds || bInsideBounds)
+	{
+		const FTransform& ComponentTransform = GeometryCollectionComponent->GetComponentTransform();
+		const FVector LocalStart = ComponentTransform.InverseTransformPosition(Start);
+		const FVector LocalEnd = ComponentTransform.InverseTransformPosition(End);
+		float MinDistance = FLT_MAX;
+		FVector Intersect;
+		FVector Normal;
+		UE::Geometry::FIndex3i FoundTriangle;
+		FVector HitPosition;
+		if (!RayIntersectAdapter(FoundTriangle, HitPosition, LocalStart, LocalEnd))
+		{
+			return false;
+		}
+
+		// Compute the normal of the triangle
+		const FVector& P0 = MeshVertices[FoundTriangle.A];
+		const FVector& P1 = MeshVertices[FoundTriangle.B];
+		const FVector& P2 = MeshVertices[FoundTriangle.C];
+
+		const FVector TriNorm = (P1 - P0) ^ (P2 - P0);
+
+		//check collinearity of A,B,C
+		if (TriNorm.SizeSquared() > SMALL_NUMBER)
+		{
+			FVector IntersectPoint;
+			FVector HitNormal;
+		
+			bool bHit = FMath::SegmentTriangleIntersection(LocalStart, LocalEnd, P0, P1, P2, IntersectPoint, HitNormal);
+
+			if (bHit)
+			{
+				const float Distance = (LocalStart - IntersectPoint).SizeSquared();
+				if (Distance < MinDistance)
+				{
+					MinDistance = Distance;
+					Intersect = IntersectPoint;
+					Normal = HitNormal;
+				}
+			}
+		}
+		
+
+		if (MinDistance != FLT_MAX)
+		{
+			OutHit.Component = GeometryCollectionComponent;
+			OutHit.Normal = ComponentTransform.TransformVector(Normal).GetSafeNormal();
+			OutHit.ImpactNormal = OutHit.Normal;
+			OutHit.ImpactPoint = ComponentTransform.TransformPosition(Intersect);
+			OutHit.Location = OutHit.ImpactPoint;
+			OutHit.bBlockingHit = true;
+			OutHit.Distance = MinDistance;
+			bHitTriangle = true;
+		}
+	}	
+
+	return bHitTriangle;
+}
+
+void FMeshPaintGeometryCollectionComponentAdapter::QueryPaintableTextures(int32 MaterialIndex, int32& OutDefaultIndex, TArray<struct FPaintableTexture>& InOutTextureList)
+{
+	if (GeometryCollectionComponent.IsValid())
+	{
+		DefaultQueryPaintableTextures(MaterialIndex, GeometryCollectionComponent.Get(), OutDefaultIndex, InOutTextureList);
+	}
+}
+
+void FMeshPaintGeometryCollectionComponentAdapter::ApplyOrRemoveTextureOverride(UTexture* SourceTexture, UTexture* OverrideTexture) const
+{
+	if (GeometryCollectionComponent.IsValid())
+	{
+		DefaultApplyOrRemoveTextureOverride(GeometryCollectionComponent.Get(), SourceTexture, OverrideTexture);
+	}
+}
+
+void FMeshPaintGeometryCollectionComponentAdapter::GetTextureCoordinate(int32 VertexIndex, int32 ChannelIndex, FVector2D& OutTextureCoordinate) const
+{
+	const TArray<FVector2f>& UV = GetGeometryCollectionObject()->GetGeometryCollection()->UVs[VertexIndex];
+	if (ChannelIndex < UV.Num())
+	{
+		OutTextureCoordinate = (FVector2D)UV[ChannelIndex];
+	}
+}
+
+void FMeshPaintGeometryCollectionComponentAdapter::PreEdit()
+{
+	if (!IsValid())
+	{
+		return;
+	}
+
+	FlushRenderingCommands();
+
+	GeometryCollectionComponent->Modify();
+	GetGeometryCollectionObject()->Modify();
+}
+
+void FMeshPaintGeometryCollectionComponentAdapter::PostEdit()
+{
+	if (!GeometryCollectionComponent.IsValid())
+	{
+		return;
+	}
+
+	GeometryCollectionComponent->MarkRenderStateDirty();
+}
+
+void FMeshPaintGeometryCollectionComponentAdapter::GetVertexColor(int32 VertexIndex, FColor& OutColor, bool bInstance /*= true*/) const
+{
+	OutColor = GetGeometryCollectionObject()->GetGeometryCollection()->Color[VertexIndex].ToFColor(true);
+}
+
+void FMeshPaintGeometryCollectionComponentAdapter::SetVertexColor(int32 VertexIndex, FColor Color, bool bInstance /*= true*/)
+{
+	GetGeometryCollectionObject()->GetGeometryCollection()->Color[VertexIndex] = FLinearColor(Color);
+}
+
+FMatrix FMeshPaintGeometryCollectionComponentAdapter::GetComponentToWorldMatrix() const
+{
+	if (!GeometryCollectionComponent.IsValid())
+	{
+		return FMatrix::Identity;
+	}
+
+	return GeometryCollectionComponent->GetComponentToWorld().ToMatrixWithScale();
+}
+
+//////////////////////////////////////////////////////////////////////////
+// FMeshPaintGeometryCollectionComponentAdapterFactory
+
+TSharedPtr<IMeshPaintComponentAdapter> FMeshPaintGeometryCollectionComponentAdapterFactory::Construct(class UMeshComponent* InComponent, int32 InMeshLODIndex) const
+{
+	if (UGeometryCollectionComponent* GeometryCollectionComponent = Cast<UGeometryCollectionComponent>(InComponent))
+	{
+		if (GeometryCollectionComponent->GetRestCollection() != nullptr)
+		{
+			TSharedRef<FMeshPaintGeometryCollectionComponentAdapter> Result = MakeShareable(new FMeshPaintGeometryCollectionComponentAdapter());
+			if (Result->Construct(InComponent, InMeshLODIndex))
+			{
+				return Result;
+			}
+		}
+	}
+
+	return nullptr;
+}
diff --git a/Engine/Plugins/MeshPainting/Source/MeshPaintingToolset/Public/MeshPaintGeometryCollectionAdapter.h b/Engine/Plugins/MeshPainting/Source/MeshPaintingToolset/Public/MeshPaintGeometryCollectionAdapter.h
new file mode 100644
--- /dev/null
+++ b/Engine/Plugins/MeshPainting/Source/MeshPaintingToolset/Public/MeshPaintGeometryCollectionAdapter.h
@@ -0,0 +1,80 @@
+// Copyright Epic Games, Inc. All Rights Reserved.
+
+#pragma once
+
+#include "CoreMinimal.h"
+#include "BaseMeshPaintComponentAdapter.h"
+#include "MeshPaintComponentAdapterFactory.h"
+#include "UObject/WeakObjectPtrTemplates.h"
+#include "UObject/WeakObjectPtr.h"
+
+class UBodySetup;
+class UGeometryCollection;
+class UGeometryCollectionComponent;
+class UTexture;
+class UGeometryCollection;
+class UMeshComponent;
+
+//////////////////////////////////////////////////////////////////////////
+// FMeshPaintGeometryAdapterForSkeletalMeshes
+class MESHPAINTINGTOOLSET_API FMeshPaintGeometryCollectionComponentAdapter : public FBaseMeshPaintComponentAdapter
+{
+public:
+	static void InitializeAdapterGlobals();
+	static void AddReferencedObjectsGlobals(FReferenceCollector& Collector);
+	static void CleanupGlobals();
+	/** Start IMeshPaintGeometryAdapter Overrides */
+	virtual bool Construct(UMeshComponent* InComponent, int32 InMeshLODIndex) override;
+	virtual ~FMeshPaintGeometryCollectionComponentAdapter();
+	virtual bool Initialize() override;
+	virtual void OnAdded() override;
+	virtual void OnRemoved() override;
+	virtual bool IsValid() const override;
+	virtual bool SupportsTexturePaint() const override { return IsValid(); }
+	virtual bool SupportsVertexPaint() const override { return IsValid(); }
+	virtual bool LineTraceComponent(struct FHitResult& OutHit, const FVector Start, const FVector End, const struct FCollisionQueryParams& Params) const override;
+	virtual void QueryPaintableTextures(int32 MaterialIndex, int32& OutDefaultIndex, TArray<struct FPaintableTexture>& InOutTextureList) override;
+	virtual void ApplyOrRemoveTextureOverride(UTexture* SourceTexture, UTexture* OverrideTexture) const override;
+	virtual void AddReferencedObjects(FReferenceCollector& Collector) override;
+	virtual void PreEdit() override;
+	virtual void PostEdit() override;
+	virtual void GetTextureCoordinate(int32 VertexIndex, int32 ChannelIndex, FVector2D& OutTextureCoordinate) const override;
+	virtual void GetVertexColor(int32 VertexIndex, FColor& OutColor, bool bInstance = true) const override;
+	virtual void SetVertexColor(int32 VertexIndex, FColor Color, bool bInstance = true) override;
+	virtual FMatrix GetComponentToWorldMatrix() const override;
+	/** End IMeshPaintGeometryAdapter Overrides */
+		
+	/** Start FBaseMeshPaintGeometryAdapter Overrides */
+	virtual bool InitializeVertexData();
+	/** End FBaseMeshPaintGeometryAdapter Overrides */
+protected:
+
+	/** Delegate called when geometry collection is changed on the component */
+	FDelegateHandle GeometryCollectionChangedHandle;
+
+	/** Geometry Collection component represented by this adapter */
+	TWeakObjectPtr<UGeometryCollectionComponent> GeometryCollectionComponent;
+	/// Get the underlying UGeometryCollection from the component, as a non-const object
+	/// Caller must have already validated that the component weak pointer is still valid (as this is called per-vertex)
+	UGeometryCollection* GetGeometryCollectionObject() const;
+
+	void OnGeometryCollectionChanged();
+
+
+	// Like IsValid() but does not verify that the cached data matches
+	bool HasValidGeometryCollection() const;
+	
+	// TODO: Store a LOD index if/when GeometryCollection supports LODs
+};
+
+//////////////////////////////////////////////////////////////////////////
+// FMeshPaintGeometryAdapterForSkeletalMeshesFactory
+
+class MESHPAINTINGTOOLSET_API FMeshPaintGeometryCollectionComponentAdapterFactory : public IMeshPaintComponentAdapterFactory
+{
+public:
+	virtual TSharedPtr<IMeshPaintComponentAdapter> Construct(UMeshComponent* InComponent, int32 InMeshLODIndex) const override;
+	virtual void InitializeAdapterGlobals() override { FMeshPaintGeometryCollectionComponentAdapter::InitializeAdapterGlobals(); }
+	virtual void AddReferencedObjectsGlobals(FReferenceCollector& Collector) override { FMeshPaintGeometryCollectionComponentAdapter::AddReferencedObjectsGlobals(Collector); }
+	virtual void CleanupGlobals() override { FMeshPaintGeometryCollectionComponentAdapter::CleanupGlobals(); }
+};
diff --git a/Engine/Source/Runtime/Experimental/GeometryCollectionEngine/Private/GeometryCollection/GeometryCollectionComponent.cpp b/Engine/Source/Runtime/Experimental/GeometryCollectionEngine/Private/GeometryCollection/GeometryCollectionComponent.cpp
--- a/Engine/Source/Runtime/Experimental/GeometryCollectionEngine/Private/GeometryCollection/GeometryCollectionComponent.cpp
+++ b/Engine/Source/Runtime/Experimental/GeometryCollectionEngine/Private/GeometryCollection/GeometryCollectionComponent.cpp
@@ -274,8 +274,6 @@ UGeometryCollectionComponent::UGeometryCollectionComponent(const FObjectInitiali
 	GlobalNavMeshInvalidationCounter += 3;
 	NavmeshInvalidationTimeSliceIndex = GlobalNavMeshInvalidationCounter;
 
-	WorldBounds = FBoxSphereBounds(FBox(ForceInit));	
-
 	// default current cache time
 	CurrentCacheTime = MAX_flt;
 
@@ -398,6 +396,28 @@ void UGeometryCollectionComponent::GetLifetimeReplicatedProps(TArray<FLifetimePr
 	DOREPLIFETIME(UGeometryCollectionComponent, RepData);
 }
 
+namespace
+{
+void UpdateGlobalMatricesWithExplodedVectors(TArray<FMatrix>& GlobalMatricesIn, FGeometryCollection& GeometryCollection)
+{
+	int32 NumMatrices = GlobalMatricesIn.Num();
+	if (GlobalMatricesIn.Num() > 0)
+	{
+		if (GeometryCollection.HasAttribute("ExplodedVector", FGeometryCollection::TransformGroup))
+		{
+			const TManagedArray<FVector3f>& ExplodedVectors = GeometryCollection.GetAttribute<FVector3f>("ExplodedVector", FGeometryCollection::TransformGroup);
+
+			check(NumMatrices == ExplodedVectors.Num());
+
+			for (int32 tt = 0, nt = NumMatrices; tt < nt; ++tt)
+			{
+				GlobalMatricesIn[tt] = GlobalMatricesIn[tt].ConcatTranslation((FVector)ExplodedVectors[tt]);
+			}
+		}
+	}
+}
+}
+
 FBox UGeometryCollectionComponent::ComputeBounds(const FMatrix& LocalToWorldWithScale) const
 {
 	FBox BoundingBox(ForceInit);
@@ -443,6 +463,7 @@ FBox UGeometryCollectionComponent::ComputeBounds(const FMatrix& LocalToWorldWith
 			}
 			else
 			{
+				UpdateGlobalMatricesWithExplodedVectors(TmpGlobalMatrices, *(RestCollection->GetGeometryCollection()));
 				for (int32 BoxIdx = 0; BoxIdx < NumBoxes; ++BoxIdx)
 				{
 					const int32 TransformIndex = TransformIndices[BoxIdx];
@@ -530,16 +551,17 @@ FBoxSphereBounds UGeometryCollectionComponent::CalcBounds(const FTransform& Loca
 		return ComponentSpaceBounds.TransformBy(LocalToWorldIn);
 	}
 
-	// non cached bounds path 
-	if (!CachePlayback && WorldBounds.GetSphere().W > 1e-5 && NumTransforms > 0)
-	{
-		return WorldBounds;
-	} 
-
 	const FMatrix LocalToWorldWithScale = LocalToWorldIn.ToMatrixWithScale();
 	return FBoxSphereBounds(ComputeBounds(LocalToWorldWithScale));
 }
 
+void UGeometryCollectionComponent::UpdateCachedBounds()
+{
+	ComponentSpaceBounds = ComputeBounds(FMatrix::Identity);
+	CalculateLocalBounds();
+	UpdateBounds();
+}
+
 void UGeometryCollectionComponent::CreateRenderState_Concurrent(FRegisterComponentContext* Context)
 {
 	Super::CreateRenderState_Concurrent(Context);
@@ -1199,6 +1221,16 @@ void UGeometryCollectionComponent::InitializeComponent()
 }
 
 #if WITH_EDITOR
+FDelegateHandle UGeometryCollectionComponent::RegisterOnGeometryCollectionPropertyChanged(const FOnGeometryCollectionPropertyChanged& Delegate)
+{
+	return OnGeometryCollectionPropertyChanged.Add(Delegate);
+}
+
+void UGeometryCollectionComponent::UnregisterOnGeometryCollectionPropertyChanged(FDelegateHandle Handle)
+{
+	OnGeometryCollectionPropertyChanged.Remove(Handle);
+}
+
 void UGeometryCollectionComponent::PostEditChangeChainProperty(FPropertyChangedChainEvent& PropertyChangedEvent)
 {
 	Super::PostEditChangeChainProperty(PropertyChangedEvent);
@@ -1206,7 +1238,17 @@ void UGeometryCollectionComponent::PostEditChangeChainProperty(FPropertyChangedC
 	if (PropertyChangedEvent.Property && PropertyChangedEvent.Property->GetFName() == GET_MEMBER_NAME_CHECKED(UGeometryCollectionComponent, bShowBoneColors))
 	{
 		FScopedColorEdit EditBoneColor(this, true /*bForceUpdate*/); // the property has already changed; this will trigger the color update + render state updates
-	}					
+	}
+}
+
+void UGeometryCollectionComponent::PostEditChangeProperty(FPropertyChangedEvent& PropertyChangedEvent)
+{
+	Super::PostEditChangeProperty(PropertyChangedEvent);
+
+	if (OnGeometryCollectionPropertyChanged.IsBound())
+	{
+		OnGeometryCollectionPropertyChanged.Broadcast();
+	}
 }
 #endif
 
@@ -3269,20 +3311,7 @@ AChaosSolverActor* UGeometryCollectionComponent::GetPhysicsSolverActor() const
 void UGeometryCollectionComponent::CalculateLocalBounds()
 {
 	LocalBounds.Init();
-	const TManagedArray<FBox>& BoundingBoxes = GetBoundingBoxArray();
-	const TManagedArray<int32>& TransformIndices = GetTransformIndexArray();
-
-	const int32 NumBoxes = BoundingBoxes.Num();
-
-	for (int32 BoxIdx = 0; BoxIdx < NumBoxes; ++BoxIdx)
-	{
-		const int32 TransformIndex = TransformIndices[BoxIdx];
-
-		if (GetRestCollection()->GetGeometryCollection()->IsGeometry(TransformIndex))
-		{
-			LocalBounds += BoundingBoxes[BoxIdx];
-		}
-	}
+	LocalBounds = ComputeBounds(FMatrix::Identity);
 }
 
 void UGeometryCollectionComponent::CalculateGlobalMatrices()
@@ -3361,20 +3390,7 @@ void UGeometryCollectionComponent::CalculateGlobalMatrices()
 	}
 	
 #if WITH_EDITOR
-	if (GlobalMatrices.Num() > 0)
-	{
-		if (RestCollection->GetGeometryCollection()->HasAttribute("ExplodedVector", FGeometryCollection::TransformGroup))
-		{
-			const TManagedArray<FVector3f>& ExplodedVectors = RestCollection->GetGeometryCollection()->GetAttribute<FVector3f>("ExplodedVector", FGeometryCollection::TransformGroup);
-
-			check(GlobalMatrices.Num() == ExplodedVectors.Num());
-
-			for (int32 tt = 0, nt = GlobalMatrices.Num(); tt < nt; ++tt)
-			{
-				GlobalMatrices[tt] = GlobalMatrices[tt].ConcatTranslation((FVector)ExplodedVectors[tt]);
-			}
-		}
-	}
+	UpdateGlobalMatricesWithExplodedVectors(GlobalMatrices, *(RestCollection->GetGeometryCollection()));
 #endif
 }
 
diff --git a/Engine/Source/Runtime/Experimental/GeometryCollectionEngine/Private/GeometryCollection/GeometryCollectionSceneProxy.cpp b/Engine/Source/Runtime/Experimental/GeometryCollectionEngine/Private/GeometryCollection/GeometryCollectionSceneProxy.cpp
--- a/Engine/Source/Runtime/Experimental/GeometryCollectionEngine/Private/GeometryCollection/GeometryCollectionSceneProxy.cpp
+++ b/Engine/Source/Runtime/Experimental/GeometryCollectionEngine/Private/GeometryCollection/GeometryCollectionSceneProxy.cpp
@@ -836,6 +836,63 @@ void FGeometryCollectionSceneProxy::GetDynamicMeshElements(const TArray<const FS
 	if (GetRequiredVertexCount())
 	{
 		const bool bWireframe = AllowDebugViewmodes() && ViewFamily.EngineShowFlags.Wireframe;
+		const bool bProxyIsSelected = IsSelected();
+
+		const FEngineShowFlags& EngineShowFlags = ViewFamily.EngineShowFlags;
+
+		auto SetDebugMaterial = [this, &Collector, &EngineShowFlags, bProxyIsSelected](FMeshBatch& Mesh) -> void
+		{
+#if UE_ENABLE_DEBUG_DRAWING
+
+			// flag to indicate whether we've set a debug material yet
+			// Note: Will be used if we add more debug material options
+			// (compare to variable of same name in StaticMeshRender.cpp)
+			bool bDebugMaterialRenderProxySet = false;
+
+			if (!bDebugMaterialRenderProxySet && bProxyIsSelected && EngineShowFlags.VertexColors && AllowDebugViewmodes())
+			{
+				// Override the mesh's material with our material that draws the vertex colors
+				UMaterial* VertexColorVisualizationMaterial = NULL;
+				switch (GVertexColorViewMode)
+				{
+				case EVertexColorViewMode::Color:
+					VertexColorVisualizationMaterial = GEngine->VertexColorViewModeMaterial_ColorOnly;
+					break;
+
+				case EVertexColorViewMode::Alpha:
+					VertexColorVisualizationMaterial = GEngine->VertexColorViewModeMaterial_AlphaAsColor;
+					break;
+
+				case EVertexColorViewMode::Red:
+					VertexColorVisualizationMaterial = GEngine->VertexColorViewModeMaterial_RedOnly;
+					break;
+
+				case EVertexColorViewMode::Green:
+					VertexColorVisualizationMaterial = GEngine->VertexColorViewModeMaterial_GreenOnly;
+					break;
+
+				case EVertexColorViewMode::Blue:
+					VertexColorVisualizationMaterial = GEngine->VertexColorViewModeMaterial_BlueOnly;
+					break;
+				}
+				check(VertexColorVisualizationMaterial != NULL);
+
+				// Note: static mesh renderer does something more complicated involving per-section selection,
+				// but whole component selection seems ok for now
+				bool bSectionIsSelected = bProxyIsSelected;
+
+				auto VertexColorVisualizationMaterialInstance = new FColoredMaterialRenderProxy(
+					VertexColorVisualizationMaterial->GetRenderProxy(),
+					GetSelectionColor(FLinearColor::White, bSectionIsSelected, IsHovered())
+				);
+
+				Collector.RegisterOneFrameMaterialProxy(VertexColorVisualizationMaterialInstance);
+				Mesh.MaterialRenderProxy = VertexColorVisualizationMaterialInstance;
+
+				bDebugMaterialRenderProxySet = true;
+			}
+#endif
+		};
 
 		for (int32 ViewIndex = 0; ViewIndex < Views.Num(); ViewIndex++)
 		{
@@ -905,6 +962,9 @@ void FGeometryCollectionSceneProxy::GetDynamicMeshElements(const TArray<const FS
 						Mesh.BatchHitProxyId = Section.HitProxy ? Section.HitProxy->Id : FHitProxyId();
 					}
 				#endif
+
+					SetDebugMaterial(Mesh);
+
 					Collector.AddMesh(ViewIndex, Mesh);
 				}
 			}
@@ -954,6 +1014,9 @@ void FGeometryCollectionSceneProxy::GetDynamicMeshElements(const TArray<const FS
 						Mesh.BatchHitProxyId = Section.HitProxy ? Section.HitProxy->Id : FHitProxyId();
 					}
 				#endif
+
+					SetDebugMaterial(Mesh);
+
 					Collector.AddMesh(ViewIndex, Mesh);
 				}
 			}
diff --git a/Engine/Source/Runtime/Experimental/GeometryCollectionEngine/Public/GeometryCollection/GeometryCollectionComponent.h b/Engine/Source/Runtime/Experimental/GeometryCollectionEngine/Public/GeometryCollection/GeometryCollectionComponent.h
--- a/Engine/Source/Runtime/Experimental/GeometryCollectionEngine/Public/GeometryCollection/GeometryCollectionComponent.h
+++ b/Engine/Source/Runtime/Experimental/GeometryCollectionEngine/Public/GeometryCollection/GeometryCollectionComponent.h
@@ -435,6 +435,15 @@ public:
 
 #if WITH_EDITOR
 	virtual void PostEditChangeChainProperty(FPropertyChangedChainEvent& PropertyChangedEvent) override;
+	virtual void PostEditChangeProperty(FPropertyChangedEvent& PropertyChangedEvent) override;
+
+	DECLARE_MULTICAST_DELEGATE(FOnGeometryCollectionPropertyChangedMulticaster)
+	FOnGeometryCollectionPropertyChangedMulticaster OnGeometryCollectionPropertyChanged;
+	typedef FOnGeometryCollectionPropertyChangedMulticaster::FDelegate FOnGeometryCollectionPropertyChanged;
+
+	/** Register / Unregister delegates called when the skeletal mesh property is changed */
+	FDelegateHandle RegisterOnGeometryCollectionPropertyChanged(const FOnGeometryCollectionPropertyChanged& Delegate);
+	void UnregisterOnGeometryCollectionPropertyChanged(FDelegateHandle Handle);
 #endif
 	//~ Begin UActorComponent Interface. 
 
@@ -555,6 +564,9 @@ public:
 		return Size > 0 ? Size : DynamicCollection->NumElements(Group);	//if not, maybe dynamic has the group
 	}
 
+	// Update cached bounds; used e.g. when updating the exploded view of the geometry collection
+	void UpdateCachedBounds();
+
 	// Vertices Group
 	COPY_ON_WRITE_ATTRIBUTE(FVector3f, Vertex, FGeometryCollection::VerticesGroup) 	//GetVertexArray, GetVertexArrayCopyOnWrite, GetVertexArrayRest
 	COPY_ON_WRITE_ATTRIBUTE(TArray<FVector2f>, UVs, FGeometryCollection::VerticesGroup)		//GetUVsArray
@@ -979,7 +991,6 @@ private:
 	FBox LocalBounds;
 
 	mutable FBoxSphereBounds ComponentSpaceBounds;
-	FBoxSphereBounds WorldBounds;
 
 	float CurrentCacheTime;
 	TArray<bool> EventsPlayed;
