make mesh paint for geometry collection automatically toggle off 'bShowBoneColors' while painting

diff --git a/Engine/Plugins/MeshPainting/Source/MeshPaintingToolset/Private/MeshPaintGeometryCollectionAdapter.cpp b/Engine/Plugins/MeshPainting/Source/MeshPaintingToolset/Private/MeshPaintGeometryCollectionAdapter.cpp
--- a/Engine/Plugins/MeshPainting/Source/MeshPaintingToolset/Private/MeshPaintGeometryCollectionAdapter.cpp
+++ b/Engine/Plugins/MeshPainting/Source/MeshPaintingToolset/Private/MeshPaintGeometryCollectionAdapter.cpp
@@ -143,6 +143,8 @@ void FMeshPaintGeometryCollectionComponentAdapter::OnAdded()
 	{
 		return;
 	}
+	bSavedShowBoneColors = GeometryCollectionComponent->GetShowBoneColors();
+	GeometryCollectionComponent->SetShowBoneColors(false);
 
 	checkf(GetGeometryCollectionObject(), TEXT("Geometry Collection Component did not have a valid geometry collection object attached"));
 	checkf(GetGeometryCollectionObject()->GetGeometryCollection(), TEXT("Geometry Collection Component did not have valid geometry collection data attached"));
@@ -150,6 +152,11 @@ void FMeshPaintGeometryCollectionComponentAdapter::OnAdded()
 
 void FMeshPaintGeometryCollectionComponentAdapter::OnRemoved()
 {
+	if (GeometryCollectionComponent.IsValid())
+	{
+		GeometryCollectionComponent->SetShowBoneColors(bSavedShowBoneColors);
+		bSavedShowBoneColors = false;
+	}
 }
 
 bool FMeshPaintGeometryCollectionComponentAdapter::LineTraceComponent(struct FHitResult& OutHit, const FVector Start, const FVector End, const struct FCollisionQueryParams& Params) const
diff --git a/Engine/Plugins/MeshPainting/Source/MeshPaintingToolset/Public/MeshPaintGeometryCollectionAdapter.h b/Engine/Plugins/MeshPainting/Source/MeshPaintingToolset/Public/MeshPaintGeometryCollectionAdapter.h
--- a/Engine/Plugins/MeshPainting/Source/MeshPaintingToolset/Public/MeshPaintGeometryCollectionAdapter.h
+++ b/Engine/Plugins/MeshPainting/Source/MeshPaintingToolset/Public/MeshPaintGeometryCollectionAdapter.h
@@ -54,6 +54,9 @@ protected:
 
 	/** Geometry Collection component represented by this adapter */
 	TWeakObjectPtr<UGeometryCollectionComponent> GeometryCollectionComponent;
+
+	bool bSavedShowBoneColors = false;
+
 	/// Get the underlying UGeometryCollection from the component, as a non-const object
 	/// Caller must have already validated that the component weak pointer is still valid (as this is called per-vertex)
 	UGeometryCollection* GetGeometryCollectionObject() const;
