make mesh paint fill tool work on general adapted meshes

diff --git a/Engine/Plugins/MeshPainting/Source/MeshPaintEditorMode/Private/MeshPaintMode.cpp b/Engine/Plugins/MeshPainting/Source/MeshPaintEditorMode/Private/MeshPaintMode.cpp
--- a/Engine/Plugins/MeshPainting/Source/MeshPaintEditorMode/Private/MeshPaintMode.cpp
+++ b/Engine/Plugins/MeshPainting/Source/MeshPaintEditorMode/Private/MeshPaintMode.cpp
@@ -512,7 +512,18 @@ void UMeshPaintMode::FillWithVertexColor()
 		{
 			GEngine->GetEngineSubsystem<UMeshPaintingSubsystem>()->FillSkeletalMeshVertexColors(Cast<USkeletalMeshComponent>(Component), bPaintOnSpecificLOD ? ColorProperties->LODIndex : -1, FillColor, MaskColor);
 		}
-		
+		else if (MeshAdapter) 			// We don't have a custom fill function for this type of component; try to go through the adapter.
+		{
+			TArray<uint32> MeshIndices = MeshAdapter->GetMeshIndices();
+			UMeshPaintingSubsystem* PaintingSubsystem = GEngine->GetEngineSubsystem<UMeshPaintingSubsystem>();
+			for (uint32 VID : MeshIndices)
+			{
+				FColor Color;
+				MeshAdapter->GetVertexColor((int32)VID, Color);
+				PaintingSubsystem->ApplyFillWithMask(Color, MaskColor, FillColor);
+				MeshAdapter->SetVertexColor((int32)VID, Color);
+			}
+		}
 
 		if (MeshAdapter)
 		{
