Revise text in the Mesh Paint Mode warning dialogs.

diff --git a/Engine/Plugins/MeshPainting/Source/MeshPaintEditorMode/Private/MeshPaintMode.cpp b/Engine/Plugins/MeshPainting/Source/MeshPaintEditorMode/Private/MeshPaintMode.cpp
--- a/Engine/Plugins/MeshPainting/Source/MeshPaintEditorMode/Private/MeshPaintMode.cpp
+++ b/Engine/Plugins/MeshPainting/Source/MeshPaintEditorMode/Private/MeshPaintMode.cpp
@@ -536,12 +536,12 @@ void UMeshPaintMode::FillWithVertexColor()
 void UMeshPaintMode::PropagateVertexColorsToAsset()
 {
 	const TArray<UStaticMeshComponent*> StaticMeshComponents = GetSelectedComponents<UStaticMeshComponent>();
-	FSuppressableWarningDialog::FSetupInfo SetupInfo(LOCTEXT("PushInstanceVertexColorsPrompt_Message", "Copying the instance vertex colors to the source mesh will replace any of the source mesh's pre-existing vertex colors and affect every instance of the source mesh."),
-		LOCTEXT("PushInstanceVertexColorsPrompt_Title", "Warning: Copying vertex data overwrites all instances"), "Warning_PushInstanceVertexColorsPrompt");
+	FSuppressableWarningDialog::FSetupInfo SetupInfo(LOCTEXT("PushInstanceVertexColorsPrompt_Message", "This operation copies vertex colors from LOD 0 of the selected instance to all LODs of the source asset, overwriting any existing vertex colors.\n\nThis change will also propagate to all other instances of the same asset that do not have custom vertex colors."),
+		LOCTEXT("PushInstanceVertexColorsPrompt_Title", "Warning: Overwriting Vertex Colors on Source Asset"), "Warning_PushInstanceVertexColorsPrompt");
 
-	SetupInfo.ConfirmText = LOCTEXT("PushInstanceVertexColorsPrompt_ConfirmText", "Continue");
-	SetupInfo.CancelText = LOCTEXT("PushInstanceVertexColorsPrompt_CancelText", "Abort");
-	SetupInfo.CheckBoxText = LOCTEXT("PushInstanceVertexColorsPrompt_CheckBoxText", "Always copy vertex colors without prompting");
+	SetupInfo.ConfirmText = LOCTEXT("PushInstanceVertexColorsPrompt_ConfirmText", "Overwrite");
+	SetupInfo.CancelText = LOCTEXT("PushInstanceVertexColorsPrompt_CancelText", "Cancel");
+	SetupInfo.CheckBoxText = LOCTEXT("PushInstanceVertexColorsPrompt_CheckBoxText", "Always overwrite source asset without prompting");
 
 	FSuppressableWarningDialog VertexColorCopyWarning(SetupInfo);
 
@@ -756,12 +756,12 @@ void UMeshPaintMode::PropagateVertexColorsToLODs()
 	if (bSelectionContainsPerLODColors)
 	{
 		//Warn the user they will lose custom painting data
-		FSuppressableWarningDialog::FSetupInfo SetupInfo(LOCTEXT("LooseLowersLODsVertexColorsPrompt_Message", "Propagating Vertex Colors from LOD0 to all lower LODs. This mean all lower LODs custom vertex painting will be lost."),
-			LOCTEXT("LooseLowersLODsVertexColorsPrompt_Title", "Warning: Lowers LODs custom vertex painting will be lost!"), "Warning_LooseLowersLODsVertexColorsPrompt");
+		FSuppressableWarningDialog::FSetupInfo SetupInfo(LOCTEXT("LooseLowersLODsVertexColorsPrompt_Message", "This operation copies vertex colors from LOD 0 to all other LODs in this instance, overwriting any existing vertex colors.\n\nAt least one LOD has custom vertex colors that will be lost."),
+			LOCTEXT("LooseLowersLODsVertexColorsPrompt_Title", "Warning: Overwriting Vertex Colors on LODs"), "Warning_LooseLowersLODsVertexColorsPrompt");
 
-		SetupInfo.ConfirmText = LOCTEXT("LooseLowersLODsVertexColorsPrompt_ConfirmText", "Continue");
-		SetupInfo.CancelText = LOCTEXT("LooseLowersLODsVertexColorsPrompt_CancelText", "Abort");
-		SetupInfo.CheckBoxText = LOCTEXT("LooseLowersLODsVertexColorsPrompt_CheckBoxText", "Always copy vertex colors without prompting");
+		SetupInfo.ConfirmText = LOCTEXT("LooseLowersLODsVertexColorsPrompt_ConfirmText", "Overwrite");
+		SetupInfo.CancelText = LOCTEXT("LooseLowersLODsVertexColorsPrompt_CancelText", "Cancel");
+		SetupInfo.CheckBoxText = LOCTEXT("LooseLowersLODsVertexColorsPrompt_CheckBoxText", "Always overwrite LODs without prompting");
 
 		FSuppressableWarningDialog LooseLowersLODsVertexColorsWarning(SetupInfo);
 	
