Fixed a crash when exiting the paint mode after deleting a material that was being painted

diff --git a/Engine/Plugins/MeshPainting/Source/MeshPaintingToolset/Private/MeshTexturePaintingTool.cpp b/Engine/Plugins/MeshPainting/Source/MeshPaintingToolset/Private/MeshTexturePaintingTool.cpp
--- a/Engine/Plugins/MeshPainting/Source/MeshPaintingToolset/Private/MeshTexturePaintingTool.cpp
+++ b/Engine/Plugins/MeshPainting/Source/MeshPaintingToolset/Private/MeshTexturePaintingTool.cpp
@@ -1274,8 +1274,10 @@ void UMeshTexturePaintingTool::ClearAllTextureOverrides()
 
 		for (int32 MaterialIndex = 0; MaterialIndex < TextureData->PaintingMaterials.Num(); MaterialIndex++)
 		{
-			UMaterialInterface* PaintingMaterialInterface = TextureData->PaintingMaterials[MaterialIndex];
-			PaintingMaterialInterface->OverrideTexture(TextureData->PaintingTexture2D, nullptr, FeatureLevel);//findme
+			if (UMaterialInterface* PaintingMaterialInterface = TextureData->PaintingMaterials[MaterialIndex])
+			{
+				PaintingMaterialInterface->OverrideTexture(TextureData->PaintingTexture2D, nullptr, FeatureLevel);
+			}
 		}
 
 		TextureData->PaintingMaterials.Empty();
