Fixed the broken transaction in the texture paint tools.
Fixed the issue where exiting the tools would always commit the changes.

Additional improvement: The tool no longer create a new texture each time we start painting a stroke.

Note: the tools activation transactions still seems broken and it will need to be resolved in another change.

diff --git a/Engine/Plugins/MeshPainting/Source/MeshPaintingToolset/Private/MeshTexturePaintingTool.cpp b/Engine/Plugins/MeshPainting/Source/MeshPaintingToolset/Private/MeshTexturePaintingTool.cpp
--- a/Engine/Plugins/MeshPainting/Source/MeshPaintingToolset/Private/MeshTexturePaintingTool.cpp
+++ b/Engine/Plugins/MeshPainting/Source/MeshPaintingToolset/Private/MeshTexturePaintingTool.cpp
@@ -22,21 +22,102 @@
 #include "RHI.h"
 #include "TextureCompiler.h"
 #include "TextureResource.h"
+#include "Editor/TransBuffer.h"
+#include "Editor/UnrealEdEngine.h"
+#include "UnrealEdGlobals.h"
 
 #include UE_INLINE_GENERATED_CPP_BY_NAME(MeshTexturePaintingTool)
 
 #define LOCTEXT_NAMESPACE "MeshTextureBrush"
 
+namespace UE::MeshPaintingToolset::Private::TexturePainting
+{
+
+	// Compare two sorted arrays of unique elements
+	template<typename T>
+	void CompareSortedArrayElements(const TArrayView<T>& FirstArray, const TArrayView<T>& SecondArray, TFunctionRef<void(T&)> MissingInFirstArray, TFunctionRef<void(T&)> MissingInSecondArray,  TFunctionRef<void(T&)> PresentInBoth)
+	{
+		if (FirstArray.IsEmpty())
+		{
+			for (T& Element : SecondArray)
+			{
+				MissingInFirstArray(Element);
+			}
+			return;
+		}
+
+		if (SecondArray.IsEmpty())
+		{
+			for (T& Element : FirstArray)
+			{
+				MissingInSecondArray(Element);
+			}
+			return;
+		}
+
+
+		int32 FirstIndex = 0;
+		int32 SecondIndex = 0;
+		T* ElementFromFirstArray = nullptr;
+		T* ElementFromSecondArray = nullptr;
+
+		while (FirstIndex < FirstArray.Num() && SecondIndex < SecondArray.Num())
+		{
+			if (FirstArray.IsValidIndex(FirstIndex))
+			{
+				ElementFromFirstArray = &FirstArray[FirstIndex];
+			}
+
+			if (SecondArray.IsValidIndex(SecondIndex))
+			{
+				ElementFromSecondArray = &SecondArray[SecondIndex];
+			}
+
+			if (*ElementFromFirstArray == *ElementFromSecondArray)
+			{
+				++FirstIndex;
+				++SecondIndex;
+				PresentInBoth(*ElementFromFirstArray);
+			}
+			else if (*ElementFromFirstArray < *ElementFromSecondArray)
+			{
+				MissingInSecondArray(*ElementFromFirstArray);
+				++FirstIndex;
+			}
+			else
+			{
+				MissingInFirstArray(*ElementFromSecondArray);
+				++SecondIndex;
+			}
+		}
+
+		if (FirstArray.Num() < SecondArray.Num())
+		{
+			++SecondIndex;
+			while (SecondArray.IsValidIndex(SecondIndex))
+			{
+				MissingInFirstArray(SecondArray[SecondIndex]);
+				++SecondIndex;
+			}
+		}
+		else if (SecondArray.Num() < FirstArray.Num())
+		{
+			++FirstIndex;
+			while (FirstArray.IsValidIndex(FirstIndex))
+			{
+				MissingInSecondArray(FirstArray[FirstIndex]);
+				++FirstIndex;
+			}
+		}
+	}
+}
+
 /*
  * ToolBuilder
  */
 
-
-
 bool UMeshTexturePaintingToolBuilder::CanBuildTool(const FToolBuilderState& SceneState) const
 {
-
-
 	return GEngine->GetEngineSubsystem<UMeshPaintingSubsystem>()->SelectionHasMaterialValidForTexturePaint();
 }
 
@@ -73,6 +154,12 @@ UMeshTexturePaintingTool::UMeshTexturePaintingTool()
 void UMeshTexturePaintingTool::Setup()
 {
 	Super::Setup();
+
+	if (UTransBuffer* TransBuffer = GUnrealEd ? Cast<UTransBuffer>(GUnrealEd->Trans) : nullptr)
+	{
+		TransBuffer->OnTransactionStateChanged().AddUObject(this, &UMeshTexturePaintingTool::OnTransactionStateChanged);
+	}
+
 	TextureProperties = Cast<UMeshTexturePaintingToolProperties>(BrushProperties);
 	bResultValid = false;
 	bStampPending = false;
@@ -99,15 +186,8 @@ void UMeshTexturePaintingTool::Setup()
 void UMeshTexturePaintingTool::Shutdown(EToolShutdownType ShutdownType)
 {
 	FinishPainting();
-	// If the user has pending changes and the editor is not exiting, we want to do the commit for all the modified textures.
-	if ((GetNumberOfPendingPaintChanges() > 0) && !IsEngineExitRequested() && ShutdownType != EToolShutdownType::Cancel)
-	{
-		CommitAllPaintedTextures();
-	}
-	else
-	{
-		ClearAllTextureOverrides();
-	}
+	
+	ClearAllTextureOverrides();
 
 	PaintTargetData.Empty();
 
@@ -119,6 +199,12 @@ void UMeshTexturePaintingTool::Shutdown(EToolShutdownType ShutdownType)
 	}
 
 	BrushProperties->SaveProperties(this);
+
+	if (UTransBuffer* TransBuffer = GUnrealEd ? Cast<UTransBuffer>(GUnrealEd->Trans) : nullptr)
+	{
+		TransBuffer->OnTransactionStateChanged().RemoveAll(this);
+	}
+
 	Super::Shutdown(ShutdownType);
 }
 
@@ -478,7 +564,22 @@ bool UMeshTexturePaintingTool::PaintInternal(const TArrayView<TPair<FVector, FVe
 	return bPaintApplied;
 }
 
+void UMeshTexturePaintingTool::AddTextureOverrideToComponent(FPaintTexture2DData& TextureData, UMeshComponent* MeshComponent, const IMeshPaintComponentAdapter* MeshPaintAdapter)
+{
+	if (MeshComponent && MeshPaintAdapter)
+	{
+		if (!TextureData.PaintedComponents.Contains(MeshComponent))
+		{
+			TextureData.PaintedComponents.AddUnique(MeshComponent);
+
+			MeshPaintAdapter->ApplyOrRemoveTextureOverride(TextureData.PaintingTexture2D, TextureData.PaintRenderTargetTexture);
 
+			// For the transactions add to the cache of overridden component.
+			PaintComponentsOverride.FindOrAdd(TextureData.PaintingTexture2D).PaintedComponents.Add(MeshComponent);
+		}
+	}
+
+}
 
 void UMeshTexturePaintingTool::UpdateResult()
 {
@@ -581,8 +682,8 @@ void UMeshTexturePaintingTool::FinishPainting()
 	if (bArePainting)
 	{
 		bArePainting = false;
-		GetToolManager()->EndUndoTransaction();
 		OnPaintingFinishedDelegate.ExecuteIfBound();
+		PaintingTransaction.Reset();
 	}
 }
 
@@ -647,6 +748,9 @@ void UMeshTexturePaintingTool::StartPaintingTexture(UMeshComponent* InMeshCompon
 	check(TexturePaintingCurrentMeshComponent == nullptr);
 	check(PaintingTexture2D == nullptr);
 
+	PaintingTransaction = MakeUnique<FScopedTransaction>(LOCTEXT("MeshPaintMode_TexturePaint_Transaction", "Texture Paint"));
+	Modify();
+
 	const auto FeatureLevel = InMeshComponent->GetWorld()->FeatureLevel;
 
 	UTexture2D* Texture2D = TextureProperties->PaintTexture;
@@ -655,149 +759,134 @@ void UMeshTexturePaintingTool::StartPaintingTexture(UMeshComponent* InMeshCompon
 		return;
 	}
 
-	bool bStartedPainting = false;
 	FPaintTexture2DData* TextureData = GetPaintTargetData(Texture2D);
 
 	// Check all the materials on the mesh to see if the user texture is there
 	int32 MaterialIndex = 0;
 	UMaterialInterface* MaterialToCheck = InMeshComponent->GetMaterial(MaterialIndex);
-	while (MaterialToCheck != nullptr)
+
+	bool bIsTextureUsed = UTexturePaintToolset::DoesMeshComponentUseTexture(InMeshComponent, Texture2D);
+
+	if (!bIsTextureUsed && (TextureData != nullptr) && (TextureData->PaintRenderTargetTexture != nullptr))
+	{
+		bIsTextureUsed = UTexturePaintToolset::DoesMeshComponentUseTexture(InMeshComponent, TextureData->PaintRenderTargetTexture);
+	}
+
+	if (bIsTextureUsed)
 	{
-		bool bIsTextureUsed = UTexturePaintToolset::DoesMeshComponentUseTexture(InMeshComponent, Texture2D);
+		FTextureCompilingManager::Get().FinishCompilation({ Texture2D });
 
-		if (!bIsTextureUsed && (TextureData != nullptr) && (TextureData->PaintRenderTargetTexture != nullptr))
+		bool bIsSourceTextureStreamedIn = Texture2D->IsFullyStreamedIn();
+
+		if (!bIsSourceTextureStreamedIn)
 		{
-			bIsTextureUsed = UTexturePaintToolset::DoesMeshComponentUseTexture(InMeshComponent, TextureData->PaintRenderTargetTexture);
+			// We found that this texture is used in one of the meshes materials but not fully loaded, we will
+			//   attempt to fully stream in the texture before we try to do anything with it.
+			Texture2D->SetForceMipLevelsToBeResident(30.0f);
+			Texture2D->WaitForStreaming();
+
+			// We do a quick sanity check to make sure it is streamed fully streamed in now.
+			bIsSourceTextureStreamedIn = Texture2D->IsFullyStreamedIn();
+
 		}
 
-		if (bIsTextureUsed && !bStartedPainting)
+		if (bIsSourceTextureStreamedIn)
 		{
-			FTextureCompilingManager::Get().FinishCompilation({ Texture2D });
+			const int32 TextureWidth = Texture2D->Source.GetSizeX();
+			const int32 TextureHeight = Texture2D->Source.GetSizeY();
 
-			bool bIsSourceTextureStreamedIn = Texture2D->IsFullyStreamedIn();
+			if (TextureData == nullptr)
+			{
+				TextureData = AddPaintTargetData(Texture2D);
+			}
+			check(TextureData != nullptr);
 
-			if (!bIsSourceTextureStreamedIn)
+			// Create our render target texture
+			if (TextureData->PaintRenderTargetTexture == nullptr ||
+				TextureData->PaintRenderTargetTexture->GetSurfaceWidth() != TextureWidth ||
+				TextureData->PaintRenderTargetTexture->GetSurfaceHeight() != TextureHeight)
 			{
-				// We found that this texture is used in one of the meshes materials but not fully loaded, we will
-				//   attempt to fully stream in the texture before we try to do anything with it.
-				Texture2D->SetForceMipLevelsToBeResident(30.0f);
-				Texture2D->WaitForStreaming();
+				TextureData->PaintRenderTargetTexture = nullptr;
+				TextureData->PaintRenderTargetTexture = NewObject<UTextureRenderTarget2D>(GetTransientPackage(), NAME_None, RF_Transient);
+				TextureData->PaintRenderTargetTexture->bNeedsTwoCopies = true;
+				const bool bForceLinearGamma = true;
+				TextureData->PaintRenderTargetTexture->InitCustomFormat(TextureWidth, TextureHeight, PF_A16B16G16R16, bForceLinearGamma);
+				TextureData->PaintRenderTargetTexture->UpdateResourceImmediate();
+			}
+			TextureData->PaintRenderTargetTexture->AddressX = Texture2D->AddressX;
+			TextureData->PaintRenderTargetTexture->AddressY = Texture2D->AddressY;
 
-				// We do a quick sanity check to make sure it is streamed fully streamed in now.
-				bIsSourceTextureStreamedIn = Texture2D->IsFullyStreamedIn();
+			const int32 BrushTargetTextureWidth = TextureWidth;
+			const int32 BrushTargetTextureHeight = TextureHeight;
 
+			// Create the rendertarget used to store our paint delta
+			if (BrushRenderTargetTexture == nullptr ||
+				BrushRenderTargetTexture->GetSurfaceWidth() != BrushTargetTextureWidth ||
+				BrushRenderTargetTexture->GetSurfaceHeight() != BrushTargetTextureHeight)
+			{
+				BrushRenderTargetTexture = nullptr;
+				BrushRenderTargetTexture = NewObject<UTextureRenderTarget2D>(GetTransientPackage(), NAME_None, RF_Transient);
+				const bool bForceLinearGamma = true;
+				BrushRenderTargetTexture->ClearColor = FLinearColor::Black;
+				BrushRenderTargetTexture->bNeedsTwoCopies = true;
+				BrushRenderTargetTexture->InitCustomFormat(BrushTargetTextureWidth, BrushTargetTextureHeight, PF_A16B16G16R16, bForceLinearGamma);
+				BrushRenderTargetTexture->UpdateResourceImmediate();
+				BrushRenderTargetTexture->AddressX = TextureData->PaintRenderTargetTexture->AddressX;
+				BrushRenderTargetTexture->AddressY = TextureData->PaintRenderTargetTexture->AddressY;
 			}
 
-			if (bIsSourceTextureStreamedIn)
+			if (TextureProperties->bEnableSeamPainting)
 			{
-				const int32 TextureWidth = Texture2D->Source.GetSizeX();
-				const int32 TextureHeight = Texture2D->Source.GetSizeY();
-
-				if (TextureData == nullptr)
-				{
-					TextureData = AddPaintTargetData(Texture2D);
-				}
-				check(TextureData != nullptr);
-
-				// Create our render target texture
-				if (TextureData->PaintRenderTargetTexture == nullptr ||
-					TextureData->PaintRenderTargetTexture->GetSurfaceWidth() != TextureWidth ||
-					TextureData->PaintRenderTargetTexture->GetSurfaceHeight() != TextureHeight)
+				// Create the rendertarget used to store a mask for our paint delta area 
+				if (BrushMaskRenderTargetTexture == nullptr ||
+					BrushMaskRenderTargetTexture->GetSurfaceWidth() != BrushTargetTextureWidth ||
+					BrushMaskRenderTargetTexture->GetSurfaceHeight() != BrushTargetTextureHeight)
 				{
-					TextureData->PaintRenderTargetTexture = nullptr;
-					TextureData->PaintRenderTargetTexture = NewObject<UTextureRenderTarget2D>(GetTransientPackage(), NAME_None, RF_Transient);
-					TextureData->PaintRenderTargetTexture->bNeedsTwoCopies = true;
+					BrushMaskRenderTargetTexture = nullptr;
+					BrushMaskRenderTargetTexture = NewObject<UTextureRenderTarget2D>(GetTransientPackage(), NAME_None, RF_Transient);
 					const bool bForceLinearGamma = true;
-					TextureData->PaintRenderTargetTexture->InitCustomFormat(TextureWidth, TextureHeight, PF_A16B16G16R16, bForceLinearGamma);
-					TextureData->PaintRenderTargetTexture->UpdateResourceImmediate();
-
-					//Duplicate the texture we are painting and store it in the transient package. This texture is a backup of the data incase we want to revert before commiting.
-					TextureData->PaintingTexture2DDuplicate = (UTexture2D*)StaticDuplicateObject(Texture2D, GetTransientPackage(), *FString::Printf(TEXT("%s_TEMP"), *Texture2D->GetName()));
+					BrushMaskRenderTargetTexture->ClearColor = FLinearColor::Black;
+					BrushMaskRenderTargetTexture->bNeedsTwoCopies = true;
+					BrushMaskRenderTargetTexture->InitCustomFormat(BrushTargetTextureWidth, BrushTargetTextureHeight, PF_B8G8R8A8, bForceLinearGamma);
+					BrushMaskRenderTargetTexture->UpdateResourceImmediate();
+					BrushMaskRenderTargetTexture->AddressX = TextureData->PaintRenderTargetTexture->AddressX;
+					BrushMaskRenderTargetTexture->AddressY = TextureData->PaintRenderTargetTexture->AddressY;
 				}
-				TextureData->PaintRenderTargetTexture->AddressX = Texture2D->AddressX;
-				TextureData->PaintRenderTargetTexture->AddressY = Texture2D->AddressY;
-
-				const int32 BrushTargetTextureWidth = TextureWidth;
-				const int32 BrushTargetTextureHeight = TextureHeight;
 
-				// Create the rendertarget used to store our paint delta
-				if (BrushRenderTargetTexture == nullptr ||
-					BrushRenderTargetTexture->GetSurfaceWidth() != BrushTargetTextureWidth ||
-					BrushRenderTargetTexture->GetSurfaceHeight() != BrushTargetTextureHeight)
+				// Create the rendertarget used to store a texture seam mask
+				if (SeamMaskRenderTargetTexture == nullptr ||
+					SeamMaskRenderTargetTexture->GetSurfaceWidth() != TextureWidth ||
+					SeamMaskRenderTargetTexture->GetSurfaceHeight() != TextureHeight)
 				{
-					BrushRenderTargetTexture = nullptr;
-					BrushRenderTargetTexture = NewObject<UTextureRenderTarget2D>(GetTransientPackage(), NAME_None, RF_Transient);
+					SeamMaskRenderTargetTexture = nullptr;
+					SeamMaskRenderTargetTexture = NewObject<UTextureRenderTarget2D>(GetTransientPackage(), NAME_None, RF_Transient);
 					const bool bForceLinearGamma = true;
-					BrushRenderTargetTexture->ClearColor = FLinearColor::Black;
-					BrushRenderTargetTexture->bNeedsTwoCopies = true;
-					BrushRenderTargetTexture->InitCustomFormat(BrushTargetTextureWidth, BrushTargetTextureHeight, PF_A16B16G16R16, bForceLinearGamma);
-					BrushRenderTargetTexture->UpdateResourceImmediate();
-					BrushRenderTargetTexture->AddressX = TextureData->PaintRenderTargetTexture->AddressX;
-					BrushRenderTargetTexture->AddressY = TextureData->PaintRenderTargetTexture->AddressY;
-				}
-
-				if (TextureProperties->bEnableSeamPainting)
-				{
-					// Create the rendertarget used to store a mask for our paint delta area 
-					if (BrushMaskRenderTargetTexture == nullptr ||
-						BrushMaskRenderTargetTexture->GetSurfaceWidth() != BrushTargetTextureWidth ||
-						BrushMaskRenderTargetTexture->GetSurfaceHeight() != BrushTargetTextureHeight)
-					{
-						BrushMaskRenderTargetTexture = nullptr;
-						BrushMaskRenderTargetTexture = NewObject<UTextureRenderTarget2D>(GetTransientPackage(), NAME_None, RF_Transient);
-						const bool bForceLinearGamma = true;
-						BrushMaskRenderTargetTexture->ClearColor = FLinearColor::Black;
-						BrushMaskRenderTargetTexture->bNeedsTwoCopies = true;
-						BrushMaskRenderTargetTexture->InitCustomFormat(BrushTargetTextureWidth, BrushTargetTextureHeight, PF_B8G8R8A8, bForceLinearGamma);
-						BrushMaskRenderTargetTexture->UpdateResourceImmediate();
-						BrushMaskRenderTargetTexture->AddressX = TextureData->PaintRenderTargetTexture->AddressX;
-						BrushMaskRenderTargetTexture->AddressY = TextureData->PaintRenderTargetTexture->AddressY;
-					}
-
-					// Create the rendertarget used to store a texture seam mask
-					if (SeamMaskRenderTargetTexture == nullptr ||
-						SeamMaskRenderTargetTexture->GetSurfaceWidth() != TextureWidth ||
-						SeamMaskRenderTargetTexture->GetSurfaceHeight() != TextureHeight)
-					{
-						SeamMaskRenderTargetTexture = nullptr;
-						SeamMaskRenderTargetTexture = NewObject<UTextureRenderTarget2D>(GetTransientPackage(), NAME_None, RF_Transient);
-						const bool bForceLinearGamma = true;
-						SeamMaskRenderTargetTexture->ClearColor = FLinearColor::Black;
-						SeamMaskRenderTargetTexture->bNeedsTwoCopies = true;
-						SeamMaskRenderTargetTexture->InitCustomFormat(BrushTargetTextureWidth, BrushTargetTextureHeight, PF_B8G8R8A8, bForceLinearGamma);
-						SeamMaskRenderTargetTexture->UpdateResourceImmediate();
-						SeamMaskRenderTargetTexture->AddressX = TextureData->PaintRenderTargetTexture->AddressX;
-						SeamMaskRenderTargetTexture->AddressY = TextureData->PaintRenderTargetTexture->AddressY;
-					}
-
-					bGenerateSeamMask = true;
+					SeamMaskRenderTargetTexture->ClearColor = FLinearColor::Black;
+					SeamMaskRenderTargetTexture->bNeedsTwoCopies = true;
+					SeamMaskRenderTargetTexture->InitCustomFormat(BrushTargetTextureWidth, BrushTargetTextureHeight, PF_B8G8R8A8, bForceLinearGamma);
+					SeamMaskRenderTargetTexture->UpdateResourceImmediate();
+					SeamMaskRenderTargetTexture->AddressX = TextureData->PaintRenderTargetTexture->AddressX;
+					SeamMaskRenderTargetTexture->AddressY = TextureData->PaintRenderTargetTexture->AddressY;
 				}
 
-				bStartedPainting = true;
+				bGenerateSeamMask = true;
 			}
-		}
 
-		// @todo MeshPaint: Here we override the textures on the mesh with the render target.  The problem is that other meshes in the scene that use
-		//    this texture do not get the override. Do we want to extend this to all other selected meshes or maybe even to all meshes in the scene?
-		if (bIsTextureUsed && bStartedPainting && !TextureData->PaintingMaterials.Contains(MaterialToCheck))
-		{
-			TextureData->PaintingMaterials.AddUnique(MaterialToCheck);
+			// MeshPaint: Here we override the textures on the mesh with the render target.  The problem is that other meshes in the scene that use
+			// this texture do not get the override. Do we want to extend this to all other selected meshes or maybe even to all meshes in the scene?
+			AddTextureOverrideToComponent(*TextureData, InMeshComponent, &GeometryInfo);
 
-			GeometryInfo.ApplyOrRemoveTextureOverride(Texture2D, TextureData->PaintRenderTargetTexture);
-		}
 
-		MaterialIndex++;
-		MaterialToCheck = InMeshComponent->GetMaterial(MaterialIndex);
-	}
 
-	if (bStartedPainting)
-	{
-		TexturePaintingCurrentMeshComponent = InMeshComponent;
+			TexturePaintingCurrentMeshComponent = InMeshComponent;
+
+			check(Texture2D != nullptr);
+			PaintingTexture2D = Texture2D;
 
-		check(Texture2D != nullptr);
-		PaintingTexture2D = Texture2D;
-		// OK, now we need to make sure our render target is filled in with data
-		UTexturePaintToolset::SetupInitialRenderTargetData(TextureData->PaintingTexture2D, TextureData->PaintRenderTargetTexture);
+			// Todo investigate if this could be done before the first draw to avoid recreating the scratch texture because of a undo.
+			// OK, now we need to make sure our render target is filled in with data
+			UTexturePaintToolset::SetupInitialRenderTargetData(*TextureData);
+		}
 	}
 }
 
@@ -1142,28 +1231,176 @@ void UMeshTexturePaintingTool::FinishPaintingTexture()
 			RenderTargetResource->ReadPixels(TexturePixels);
 
 			{
-				FScopedTransaction Transaction(LOCTEXT("MeshPaintMode_TexturePaint_Transaction", "Texture Paint"));
-
 				// For undo
-				TextureData->PaintingTexture2D->SetFlags(RF_Transactional);
-				TextureData->PaintingTexture2D->Modify();
+				TextureData->ScratchTexture->Modify();
 
-				// Store source art
-				FColor* Colors = (FColor*)TextureData->PaintingTexture2D->Source.LockMip(0);
-				check(TextureData->PaintingTexture2D->Source.CalcMipSize(0) == TexturePixels.Num() * sizeof(FColor));
-				FMemory::Memcpy(Colors, TexturePixels.GetData(), TexturePixels.Num() * sizeof(FColor));
-				TextureData->PaintingTexture2D->Source.UnlockMip(0);
+				const int32 NumPixels = TexturePixels.Num();
 
-				// If render target gamma used was 1.0 then disable SRGB for the static texture
-				TextureData->PaintingTexture2D->SRGB = FMath::Abs(RenderTargetResource->GetDisplayGamma() - 1.0f) >= KINDA_SMALL_NUMBER;
+				constexpr int32 NumMips = 1;
+				constexpr int32 NumSlices = 1;
+				// Store source art into the scratch texture
+				TextureData->ScratchTexture->Source.Init(
+						TextureData->PaintingTexture2D->Source.GetSizeX(),
+						TextureData->PaintingTexture2D->Source.GetSizeY(),
+						NumSlices,
+						NumMips,
+						TSF_BGRA8,
+						UE::Serialization::FEditorBulkData::FSharedBufferWithID(MakeSharedBufferFromArray(MoveTemp(TexturePixels)))
+					);
+
+				// Store source art into the scratch texture
+				check(TextureData->ScratchTexture->Source.CalcMipSize(0) == NumPixels * sizeof(FColor));
 
-				TextureData->PaintingTexture2D->bHasBeenPaintedInEditor = true;
+				// If render target gamma used was 1.0 then disable SRGB for the static texture
+				TextureData->ScratchTexture->SRGB = FMath::Abs(RenderTargetResource->GetDisplayGamma() - 1.0f) >= KINDA_SMALL_NUMBER;
+				TextureData->ScratchTexture->bHasBeenPaintedInEditor = true;
 			}
 		}
 
 		PaintingTexture2D = nullptr;
 		TexturePaintingCurrentMeshComponent = nullptr;
 	}
+
+	PaintingTransaction.Reset();
+}
+
+void UMeshTexturePaintingTool::OnTransactionStateChanged(const FTransactionContext& InTransactionContext, const ETransactionStateEventType InTransactionState)
+{
+	if (InTransactionState == ETransactionStateEventType::UndoRedoFinalized)
+	{
+		if (GUnrealEd)
+		{
+			if (UTransactor* Trans = GUnrealEd->Trans)
+			{
+				int32 CurrentTransactionIndex = Trans->FindTransactionIndex(InTransactionContext.TransactionId);
+				if (const FTransaction* Transaction = Trans->GetTransaction(CurrentTransactionIndex))
+				{
+					// Update the texture override of the components
+					using namespace UE::MeshPaintingToolset::Private::TexturePainting;
+
+					UMeshPaintingSubsystem* MeshPaintingSubsystem = GEngine->GetEngineSubsystem<UMeshPaintingSubsystem>();
+					check(MeshPaintingSubsystem);
+					{
+						auto TextureOverrideIsMissing = [this, MeshPaintingSubsystem](TObjectPtr<UTexture2D>& Texture)
+						{
+							if (FPaintTexture2DData* TextureData = GetPaintTargetData(Texture))
+							{
+								FPaintComponentOverride& PaintComponentOverride = PaintComponentsOverride.FindOrAdd(TextureData->PaintingTexture2D);
+								PaintComponentOverride.PaintedComponents.Reserve(TextureData->PaintedComponents.Num());
+
+								for (UMeshComponent* Component : TextureData->PaintedComponents)
+								{
+									if (TSharedPtr<IMeshPaintComponentAdapter> PaintAdapterRessource = MeshPaintingSubsystem->GetAdapterForComponent(Component))
+									{
+										PaintAdapterRessource->ApplyOrRemoveTextureOverride(TextureData->PaintingTexture2D, TextureData->PaintRenderTargetTexture);
+										PaintComponentOverride.PaintedComponents.Add(Component);
+									}
+								}
+							}
+						};
+
+						auto TextureOverrideIsNotNeeded = [this, MeshPaintingSubsystem](TObjectPtr<UTexture2D>& Texture)
+						{
+							FPaintComponentOverride ComponentOverride;
+							PaintComponentsOverride.RemoveAndCopyValue(Texture, ComponentOverride);
+
+							for (UMeshComponent* Component : ComponentOverride.PaintedComponents)
+							{
+								if (TSharedPtr<IMeshPaintComponentAdapter> PaintAdapterRessource = MeshPaintingSubsystem->GetAdapterForComponent(Component))
+								{
+									PaintAdapterRessource->ApplyOrRemoveTextureOverride(Texture,  nullptr);
+								}
+							}
+						};
+
+
+						auto UpdateOverridedComponents = [this, MeshPaintingSubsystem](TObjectPtr<UTexture2D>& Texture)
+						{
+							if (FPaintTexture2DData* TextureData = GetPaintTargetData(Texture))
+							{
+								FPaintComponentOverride& ComponentOverride = PaintComponentsOverride.FindChecked(Texture);
+
+								auto SortComponents = [](const UMeshComponent& First, const UMeshComponent& Second) { return First < Second; };
+
+								TextureData->PaintedComponents.Sort(SortComponents);
+								ComponentOverride.PaintedComponents.Sort(SortComponents);
+
+								TArray<UMeshComponent*> MeshComponentsToRemove;
+								TArray<UMeshComponent*> MeshComponentsToAdd;
+
+								auto ComponentOverrideNotNeeded = [this, MeshPaintingSubsystem, &ComponentOverride, Texture, &MeshComponentsToRemove](TObjectPtr<UMeshComponent>& MeshComponent)
+								{
+									if (TSharedPtr<IMeshPaintComponentAdapter> PaintAdapterRessource = MeshPaintingSubsystem->GetAdapterForComponent(MeshComponent))
+									{
+										PaintAdapterRessource->ApplyOrRemoveTextureOverride(Texture, nullptr);
+									}
+
+									MeshComponentsToRemove.Add(MeshComponent);
+								};
+
+								auto ComponentOverrideMissing = [this, MeshPaintingSubsystem, &ComponentOverride, Texture, &MeshComponentsToAdd, RenderTarget = TextureData->PaintRenderTargetTexture](TObjectPtr<UMeshComponent>& MeshComponent)
+								{
+									if (TSharedPtr<IMeshPaintComponentAdapter> PaintAdapterRessource = MeshPaintingSubsystem->GetAdapterForComponent(MeshComponent))
+									{
+										PaintAdapterRessource->ApplyOrRemoveTextureOverride(Texture, RenderTarget);
+									}
+
+									MeshComponentsToAdd.Add(MeshComponent);
+								};
+
+								auto DoNothing = [](TObjectPtr<UMeshComponent>&) {};
+
+								CompareSortedArrayElements<TObjectPtr<UMeshComponent>>(TextureData->PaintedComponents, ComponentOverride.PaintedComponents, ComponentOverrideNotNeeded, ComponentOverrideMissing, DoNothing);
+							}
+						};
+					
+						// Unfortunately TObjectPtrs don't behave like a RawPtrs when sorting an array 
+						auto SortPredicate = [](const UTexture2D& First, const UTexture2D& Second) {return (&First) < (&Second);};
+
+						TArray<TObjectPtr<UTexture2D>> CurrentTextures;
+						PaintTargetData.GenerateKeyArray(CurrentTextures);
+						CurrentTextures.Sort(SortPredicate);
+
+						TArray<TObjectPtr<UTexture2D>> OverrideTextures;
+						PaintComponentsOverride.GenerateKeyArray(OverrideTextures);
+						PaintComponentsOverride.KeySort(SortPredicate);
+
+						CompareSortedArrayElements<TObjectPtr<UTexture2D>>(CurrentTextures, OverrideTextures, TextureOverrideIsNotNeeded, TextureOverrideIsMissing, UpdateOverridedComponents);
+					}
+
+
+					// Update the render targets
+					{
+						if (Transaction->ContainsObject(this) && TextureProperties)
+						{
+							if (TextureProperties->PaintTexture)
+							{
+								if (FPaintTexture2DData* TextureData = GetPaintTargetData(TextureProperties->PaintTexture))
+								{
+									if (TextureData->ScratchTexture)
+									{
+										UTexturePaintToolset::UpdateRenderTargetData(*TextureData);
+									}
+								}
+							}
+						}
+						else
+						{
+							for (TPair<TObjectPtr<UTexture2D>, FPaintTexture2DData>& TagetData : PaintTargetData)
+							{
+								FPaintTexture2DData& PaintTexture2DData = TagetData.Value;
+
+								if (Transaction->ContainsObject(PaintTexture2DData.ScratchTexture))
+								{
+									UTexturePaintToolset::UpdateRenderTargetData(PaintTexture2DData);
+								}
+							}
+						}
+					}
+				}
+			}
+		}
+	}
 }
 
 void UMeshTexturePaintingTool::CycleTextures(int32 Direction)
@@ -1197,7 +1434,7 @@ void UMeshTexturePaintingTool::CommitAllPaintedTextures()
 	{
 		check(PaintingTexture2D == nullptr);
 
-		FScopedTransaction Transaction(LOCTEXT("MeshPaintMode_TexturePaint_Transaction", "Texture Paint"));
+		FScopedTransaction Transaction(LOCTEXT("MeshPaintMode_TexturePaint_Commit", "Commit Texture Paint"));
 
 	//	GWarn->BeginSlowTask(LOCTEXT("BeginMeshPaintMode_TexturePaint_CommitTask", "Committing Texture Paint Changes"), true);
 
@@ -1251,9 +1488,6 @@ void UMeshTexturePaintingTool::CommitAllPaintedTextures()
 
 					TextureData->bIsPaintingTexture2DModified = false;
 
-					// Reduplicate the duplicate so that if we cancel our future changes, it will restore to how the texture looked at this point.
-					TextureData->PaintingTexture2DDuplicate = (UTexture2D*)StaticDuplicateObject(TextureData->PaintingTexture2D, GetTransientPackage(), *FString::Printf(TEXT("%s_TEMP"), *TextureData->PaintingTexture2D->GetName()));
-
 				}
 			}
 		}
@@ -1266,21 +1500,26 @@ void UMeshTexturePaintingTool::CommitAllPaintedTextures()
 
 void UMeshTexturePaintingTool::ClearAllTextureOverrides()
 {
-	const auto FeatureLevel = GMaxRHIFeatureLevel;
-	/** Remove all texture overrides which are currently stored and active */
-	for (decltype(PaintTargetData)::TIterator It(PaintTargetData); It; ++It)
+	if (UMeshPaintingSubsystem* MeshPaintingSubsystem = GEngine->GetEngineSubsystem<UMeshPaintingSubsystem>())
 	{
-		FPaintTexture2DData* TextureData = &It.Value();
-
-		for (int32 MaterialIndex = 0; MaterialIndex < TextureData->PaintingMaterials.Num(); MaterialIndex++)
+		const auto FeatureLevel = GetTargetWorld()->FeatureLevel;
+		/** Remove all texture overrides which are currently stored and active */
+		for (decltype(PaintTargetData)::TIterator It(PaintTargetData); It; ++It)
 		{
-			if (UMaterialInterface* PaintingMaterialInterface = TextureData->PaintingMaterials[MaterialIndex])
+			FPaintTexture2DData* TextureData = &It.Value();
+
+			for (UMeshComponent* MeshComponent : TextureData->PaintedComponents)
 			{
-				PaintingMaterialInterface->OverrideTexture(TextureData->PaintingTexture2D, nullptr, FeatureLevel);
+				if (TSharedPtr<IMeshPaintComponentAdapter> PaintAdapter = MeshPaintingSubsystem->GetAdapterForComponent(MeshComponent))
+				{
+					PaintAdapter->ApplyOrRemoveTextureOverride(TextureData->PaintingTexture2D, nullptr);
+				}
 			}
+	
+			TextureData->PaintedComponents.Empty();
 		}
 
-		TextureData->PaintingMaterials.Empty();
+		PaintComponentsOverride.Empty();
 	}
 }
 
diff --git a/Engine/Plugins/MeshPainting/Source/MeshPaintingToolset/Private/TexturePaintToolset.cpp b/Engine/Plugins/MeshPainting/Source/MeshPaintingToolset/Private/TexturePaintToolset.cpp
--- a/Engine/Plugins/MeshPainting/Source/MeshPaintingToolset/Private/TexturePaintToolset.cpp
+++ b/Engine/Plugins/MeshPainting/Source/MeshPaintingToolset/Private/TexturePaintToolset.cpp
@@ -316,7 +316,7 @@ bool UTexturePaintToolset::GenerateSeamMask(UMeshComponent* MeshComponent, int32
 	return RetVal;
 }
 
-UTexture2D* UTexturePaintToolset::CreateTempUncompressedTexture(UTexture2D* SourceTexture)
+UTexture2D* UTexturePaintToolset::CreateScratchUncompressedTexture(UTexture2D* SourceTexture)
 {
 	check(SourceTexture->Source.IsValid());
 
@@ -334,6 +334,7 @@ UTexture2D* UTexturePaintToolset::CreateTempUncompressedTexture(UTexture2D* Sour
 	// Allocate the new texture
 	UTexture2D* NewTexture2D = UTexture2D::CreateTransient(Width, Height, GetTempUncompressedTexturePixelFormat());
 
+	NewTexture2D->SetFlags(RF_Transactional);
 	// Fill in the base mip for the texture we created
 	uint8* MipData = (uint8*)NewTexture2D->GetPlatformData()->Mips[0].BulkData.Lock(LOCK_READ_WRITE);
 	for (int32 y = 0; y < Height; y++)
@@ -362,36 +363,72 @@ UTexture2D* UTexturePaintToolset::CreateTempUncompressedTexture(UTexture2D* Sour
 	return NewTexture2D;
 }
 
-void UTexturePaintToolset::SetupInitialRenderTargetData(UTexture2D* InTextureSource, UTextureRenderTarget2D* InRenderTarget)
+void UTexturePaintToolset::SetupInitialRenderTargetData(FPaintTexture2DData& PaintTextureData)
 {
-	check(InTextureSource != nullptr);
-	check(InRenderTarget != nullptr);
+	check(PaintTextureData.PaintingTexture2D != nullptr);
+	check(PaintTextureData.PaintRenderTargetTexture != nullptr);
 
-	if (InTextureSource->Source.IsValid())
+	if (PaintTextureData.PaintingTexture2D->Source.IsValid())
 	{
 		// Great, we have source data!  We'll use that as our image source.
 
 		// Create a texture in memory from the source art
 		{
-			// @todo MeshPaint: This generates a lot of memory thrash -- try to cache this texture and reuse it?
-			UTexture2D* TempSourceArtTexture = CreateTempUncompressedTexture(InTextureSource);
-			check(TempSourceArtTexture != nullptr);
+			if (!PaintTextureData.ScratchTexture)
+			{
+				PaintTextureData.ScratchTexture = CreateScratchUncompressedTexture(PaintTextureData.PaintingTexture2D);
 
-			// Copy the texture to the render target using the GPU
-			CopyTextureToRenderTargetTexture(TempSourceArtTexture, InRenderTarget, GEditor->GetEditorWorldContext().World()->FeatureLevel);
+				check(PaintTextureData.ScratchTexture != nullptr);
 
-			// NOTE: TempSourceArtTexture is no longer needed (will be GC'd)
+				// Copy the texture to the render target using the GPU
+				CopyTextureToRenderTargetTexture(PaintTextureData.ScratchTexture, PaintTextureData.PaintRenderTargetTexture, GEditor->GetEditorWorldContext().World()->FeatureLevel);
+			}
+			
+			// No need to update the render target if the scratch texture exist the paint operations and the undo/redo keep the render target up to date
 		}
 	}
 	else
 	{
 		// Just copy (render) the texture in GPU memory to our render target.  Hopefully it's not
 		// compressed already!
-		check(InTextureSource->IsFullyStreamedIn());
-		CopyTextureToRenderTargetTexture(InTextureSource, InRenderTarget, GEditor->GetEditorWorldContext().World()->FeatureLevel);
+		check(PaintTextureData.PaintingTexture2D->IsFullyStreamedIn());
+		CopyTextureToRenderTargetTexture(PaintTextureData.PaintingTexture2D, PaintTextureData.PaintRenderTargetTexture, GEditor->GetEditorWorldContext().World()->FeatureLevel);
 	}
 }
 
+void UTexturePaintToolset::UpdateRenderTargetData(FPaintTexture2DData& PaintTextureData)
+{
+	check(PaintTextureData.ScratchTexture);
+	check(PaintTextureData.PaintRenderTargetTexture);
+
+
+	TArray64<uint8> RawData;
+	PaintTextureData.ScratchTexture->Source.GetMipData(RawData, 0);
+
+	int32 Width = PaintTextureData.ScratchTexture->Source.GetSizeX();
+	int32 Height = PaintTextureData.ScratchTexture->Source.GetSizeY();
+
+	// Fill in the base mip for the texture we created
+	uint8* MipData = (uint8*)PaintTextureData.ScratchTexture->GetPlatformData()->Mips[0].BulkData.Lock(LOCK_READ_WRITE);
+	for (int32 y = 0; y < Height; y++)
+	{
+		uint8* DestPtr = &MipData[(Height - 1 - y) * Width * sizeof(FColor)];
+		const FColor* SrcPtr = &((FColor*)(RawData.GetData()))[(Height - 1 - y) * Width];
+		for (int32 x = 0; x < Width; x++)
+		{
+			*DestPtr++ = SrcPtr->B;
+			*DestPtr++ = SrcPtr->G;
+			*DestPtr++ = SrcPtr->R;
+			*DestPtr++ = SrcPtr->A;
+			SrcPtr++;
+		}
+	}
+	PaintTextureData.ScratchTexture->GetPlatformData()->Mips[0].BulkData.Unlock();
+	PaintTextureData.ScratchTexture->UpdateResource();
+
+	CopyTextureToRenderTargetTexture(PaintTextureData.ScratchTexture, PaintTextureData.PaintRenderTargetTexture, GEditor->GetEditorWorldContext().World()->FeatureLevel);
+}
+
 void UTexturePaintToolset::FindMaterialIndicesUsingTexture(const UTexture* Texture, const UMeshComponent* MeshComponent, TArray<int32>& OutIndices)
 {
 	checkf(Texture && MeshComponent, TEXT("Invalid Texture of MeshComponent"));
diff --git a/Engine/Plugins/MeshPainting/Source/MeshPaintingToolset/Public/MeshPaintingToolsetTypes.h b/Engine/Plugins/MeshPainting/Source/MeshPaintingToolset/Public/MeshPaintingToolsetTypes.h
--- a/Engine/Plugins/MeshPainting/Source/MeshPaintingToolset/Public/MeshPaintingToolsetTypes.h
+++ b/Engine/Plugins/MeshPainting/Source/MeshPaintingToolset/Public/MeshPaintingToolsetTypes.h
@@ -121,53 +121,40 @@ struct FPaintTexture2DData
 	GENERATED_BODY()
 	/** The original texture that we're painting */
 	UPROPERTY(Transient)
-	TObjectPtr<UTexture2D> PaintingTexture2D;
-	bool bIsPaintingTexture2DModified;
+	TObjectPtr<UTexture2D> PaintingTexture2D = nullptr;
 
-	/** A copy of the original texture we're painting, used for restoration. */
+	bool bIsPaintingTexture2DModified = false;
+
+	/** Texture used to store the paint modifications for the transactions */
 	UPROPERTY(Transient)
-	TObjectPtr<UTexture2D> PaintingTexture2DDuplicate;
+	TObjectPtr<UTexture2D> ScratchTexture = nullptr;
 
 	/** Render target texture for painting */
 	UPROPERTY(Transient)
-	TObjectPtr<UTextureRenderTarget2D> PaintRenderTargetTexture;
+	TObjectPtr<UTextureRenderTarget2D> PaintRenderTargetTexture = nullptr;
 
-	/** Render target texture used as an input while painting that contains a clone of the original image */
+	/** List of component we are painting on */
 	UPROPERTY(Transient)
-	TObjectPtr<UTextureRenderTarget2D> CloneRenderTargetTexture;
+	TArray<TObjectPtr<UMeshComponent>> PaintedComponents;
 
-	/** List of materials we are painting on */
-	UPROPERTY(Transient)
-	TArray< TObjectPtr<UMaterialInterface> > PaintingMaterials;
-
-	/** Default ctor */
-	FPaintTexture2DData() :
-		PaintingTexture2D(NULL),
-		bIsPaintingTexture2DModified(false),
-		PaintingTexture2DDuplicate(nullptr),
-		PaintRenderTargetTexture(nullptr),
-		CloneRenderTargetTexture(nullptr)
-	{}
+	FPaintTexture2DData() = default;
 
-	FPaintTexture2DData(UTexture2D* InPaintingTexture2D, bool InbIsPaintingTexture2DModified = false) :
-		PaintingTexture2D(InPaintingTexture2D),
-		bIsPaintingTexture2DModified(InbIsPaintingTexture2DModified),
-		PaintRenderTargetTexture(nullptr),
-		CloneRenderTargetTexture(nullptr)
-	{}
-
-	/** Serializer */
-	void AddReferencedObjects(FReferenceCollector& Collector)
+	FPaintTexture2DData(UTexture2D* InPaintingTexture2D, bool InbIsPaintingTexture2DModified = false)
+		: PaintingTexture2D(InPaintingTexture2D)
+		, bIsPaintingTexture2DModified(InbIsPaintingTexture2DModified)
 	{
-		// @todo MeshPaint: We're relying on GC to clean up render targets, can we free up remote memory more quickly?
-		Collector.AddReferencedObject(PaintingTexture2D);
-		Collector.AddReferencedObject(PaintRenderTargetTexture);
-		Collector.AddReferencedObject(CloneRenderTargetTexture);
-		for (int32 Index = 0; Index < PaintingMaterials.Num(); Index++)
-		{
-			Collector.AddReferencedObject(PaintingMaterials[Index]);
-		}
 	}
+
+};
+
+USTRUCT()
+struct FPaintComponentOverride
+{
+	GENERATED_BODY()
+
+	/** List of components overridden */
+	UPROPERTY(Transient)
+	TArray<TObjectPtr<UMeshComponent>> PaintedComponents;
 };
 
 namespace MeshPaintDefs
diff --git a/Engine/Plugins/MeshPainting/Source/MeshPaintingToolset/Public/MeshTexturePaintingTool.h b/Engine/Plugins/MeshPainting/Source/MeshPaintingToolset/Public/MeshTexturePaintingTool.h
--- a/Engine/Plugins/MeshPainting/Source/MeshPaintingToolset/Public/MeshTexturePaintingTool.h
+++ b/Engine/Plugins/MeshPainting/Source/MeshPaintingToolset/Public/MeshTexturePaintingTool.h
@@ -12,7 +12,9 @@
 #include "MeshPaintHelpers.h"
 #include "MeshPaintingToolsetTypes.h"
 #include "MeshVertexPaintingTool.h"
+#include "Misc/ITransaction.h"
 #include "TexturePaintToolset.h"
+
 #include "MeshTexturePaintingTool.generated.h"
 
 
@@ -23,6 +25,7 @@ class UTextureRenderTarget2D;
 struct FTextureTargetListInfo;
 enum class EMeshPaintModeAction : uint8;
 class IMeshPaintComponentAdapter;
+class FScopedTransaction;
 
 /**
  *
@@ -158,6 +161,7 @@ protected:
 	void StartPaintingTexture(UMeshComponent* InMeshComponent, const IMeshPaintComponentAdapter& GeometryInfo);
 	void PaintTexture(FMeshPaintParameters& InParams, TArray<FTexturePaintTriangleInfo>& InInfluencedTriangles, const IMeshPaintComponentAdapter& GeometryInfo);
 	void FinishPaintingTexture();
+	void OnTransactionStateChanged(const FTransactionContext& InTransactionContext, const ETransactionStateEventType InTransactionState);
 
 protected:
 	double InitialMeshArea;
@@ -174,7 +178,9 @@ protected:
 
 private:
 	bool PaintInternal(const TArrayView<TPair<FVector, FVector>>& Rays, EMeshPaintModeAction PaintAction, float PaintStrength);
-	
+
+	void AddTextureOverrideToComponent(FPaintTexture2DData& TextureData, UMeshComponent* MeshComponent, const IMeshPaintComponentAdapter* MeshPaintAdapter = nullptr);
+
 private:
 	UPROPERTY(Transient)
 	TObjectPtr<UMeshTexturePaintingToolProperties> TextureProperties;
@@ -211,7 +217,13 @@ private:
 
 	/** Stores data associated with our paint target textures */
 	UPROPERTY(Transient)
-	TMap< TObjectPtr<UTexture2D>, FPaintTexture2DData > PaintTargetData;
+	TMap<TObjectPtr<UTexture2D>, FPaintTexture2DData> PaintTargetData;
+
+	/** Store the component overrides active for each paint target textures
+	 * Note this is not transactional because we use it as cache of the current state of the scene that we can clean/update after each transaction.
+	 */
+	UPROPERTY(Transient, NonTransactional)
+	TMap<TObjectPtr<UTexture2D>, FPaintComponentOverride> PaintComponentsOverride;
 
 	/** Texture paint: Will hold a list of texture items that we can paint on */
 	TArray<FTextureTargetListInfo> TexturePaintTargetList;
@@ -226,4 +238,7 @@ private:
 
 	/** True if we need to generate a texture seam mask used for texture dilation */
 	bool bGenerateSeamMask;
+
+	/** Hold the transaction while we are painting */
+	TUniquePtr<FScopedTransaction> PaintingTransaction;
 };
diff --git a/Engine/Plugins/MeshPainting/Source/MeshPaintingToolset/Public/TexturePaintToolset.h b/Engine/Plugins/MeshPainting/Source/MeshPaintingToolset/Public/TexturePaintToolset.h
--- a/Engine/Plugins/MeshPainting/Source/MeshPaintingToolset/Public/TexturePaintToolset.h
+++ b/Engine/Plugins/MeshPainting/Source/MeshPaintingToolset/Public/TexturePaintToolset.h
@@ -12,6 +12,8 @@ class UTexture;
 class UTexture2D;
 class UTextureRenderTarget2D;
 
+struct FPaintTexture2DData;
+
 /** Batched element parameters for texture paint shaders used for paint blending and paint mask generation */
 class FMeshPaintBatchedElementParameters : public FBatchedElementParameters
 {
@@ -82,10 +84,13 @@ public:
 	}
 
 	/** Static: Creates a temporary texture used to transfer data to a render target in memory */
-	static UTexture2D* CreateTempUncompressedTexture(UTexture2D* SourceTexture);
+	static UTexture2D* CreateScratchUncompressedTexture(UTexture2D* SourceTexture);
 
 	/** Makes sure that the render target is ready to paint on */
-	static void SetupInitialRenderTargetData(UTexture2D* InTextureSource, UTextureRenderTarget2D* InRenderTarget);
+	static void SetupInitialRenderTargetData(FPaintTexture2DData& PaintTextureData);
+
+	/** Update the render target base on scratch texture */
+	static void UpdateRenderTargetData(FPaintTexture2DData& PaintTextureData);
 
 	/** Tries to find Materials using the given Texture and retrieve the corresponding material indices from the MEsh Compon*/
 	static void FindMaterialIndicesUsingTexture(const UTexture* Texture, const UMeshComponent* MeshComponent, TArray<int32>& OutIndices);
