Updated the includes and added the missing class forward declaration.

diff --git a/Engine/Plugins/MeshPainting/Source/MeshPaintingToolset/Public/MeshPaintingToolsetTypes.h b/Engine/Plugins/MeshPainting/Source/MeshPaintingToolset/Public/MeshPaintingToolsetTypes.h
--- a/Engine/Plugins/MeshPainting/Source/MeshPaintingToolset/Public/MeshPaintingToolsetTypes.h
+++ b/Engine/Plugins/MeshPainting/Source/MeshPaintingToolset/Public/MeshPaintingToolsetTypes.h
@@ -4,10 +4,11 @@
 
 #include "Engine/Texture.h"
 #include "Engine/TextureRenderTarget2D.h"
-#include "Materials/MaterialInterface.h"
 
 #include "MeshPaintingToolsetTypes.generated.h"
 
+class UMeshComponent;
+
 class FMeshPaintParameters;
 DECLARE_MULTICAST_DELEGATE_FourParams(FApplyVertexPaintData, const FMeshPaintParameters& /* &InParams*/, const FLinearColor& /*OldColor*/, FLinearColor& /*NewColor*/, const float /*PaintAmount*/);
 
