Fixed the static analysis issues.
Added some missing stuff to the transaction of the texture paint tools.

diff --git a/Engine/Plugins/MeshPainting/Source/MeshPaintingToolset/Private/MeshTexturePaintingTool.cpp b/Engine/Plugins/MeshPainting/Source/MeshPaintingToolset/Private/MeshTexturePaintingTool.cpp
--- a/Engine/Plugins/MeshPainting/Source/MeshPaintingToolset/Private/MeshTexturePaintingTool.cpp
+++ b/Engine/Plugins/MeshPainting/Source/MeshPaintingToolset/Private/MeshTexturePaintingTool.cpp
@@ -37,6 +37,9 @@ namespace UE::MeshPaintingToolset::Private::TexturePainting
 	template<typename T>
 	void CompareSortedArrayElements(const TArrayView<T>& FirstArray, const TArrayView<T>& SecondArray, TFunctionRef<void(T&)> MissingInFirstArray, TFunctionRef<void(T&)> MissingInSecondArray,  TFunctionRef<void(T&)> PresentInBoth)
 	{
+		T* ElementFromFirstArray = nullptr;
+		T* ElementFromSecondArray = nullptr;
+
 		if (FirstArray.IsEmpty())
 		{
 			for (T& Element : SecondArray)
@@ -45,6 +48,10 @@ namespace UE::MeshPaintingToolset::Private::TexturePainting
 			}
 			return;
 		}
+		else
+		{
+			ElementFromFirstArray = &FirstArray[0];
+		}
 
 		if (SecondArray.IsEmpty())
 		{
@@ -54,12 +61,18 @@ namespace UE::MeshPaintingToolset::Private::TexturePainting
 			}
 			return;
 		}
+		else
+		{
+			ElementFromSecondArray = &SecondArray[0];
+		}
 
+		if (!ElementFromFirstArray || !ElementFromSecondArray)
+		{
+			return;
+		}
 
 		int32 FirstIndex = 0;
 		int32 SecondIndex = 0;
-		T* ElementFromFirstArray = nullptr;
-		T* ElementFromSecondArray = nullptr;
 
 		while (FirstIndex < FirstArray.Num() && SecondIndex < SecondArray.Num())
 		{
@@ -93,7 +106,10 @@ namespace UE::MeshPaintingToolset::Private::TexturePainting
 
 		if (FirstArray.Num() < SecondArray.Num())
 		{
-			++SecondIndex;
+			if (*ElementFromFirstArray != *ElementFromSecondArray)
+			{
+				++SecondIndex;
+			}
 			while (SecondArray.IsValidIndex(SecondIndex))
 			{
 				MissingInFirstArray(SecondArray[SecondIndex]);
@@ -102,7 +118,10 @@ namespace UE::MeshPaintingToolset::Private::TexturePainting
 		}
 		else if (SecondArray.Num() < FirstArray.Num())
 		{
-			++FirstIndex;
+			if (*ElementFromFirstArray != *ElementFromSecondArray)
+			{
+				++FirstIndex;
+			}
 			while (FirstArray.IsValidIndex(FirstIndex))
 			{
 				MissingInSecondArray(FirstArray[FirstIndex]);
@@ -1435,7 +1454,8 @@ void UMeshTexturePaintingTool::CommitAllPaintedTextures()
 		check(PaintingTexture2D == nullptr);
 
 		FScopedTransaction Transaction(LOCTEXT("MeshPaintMode_TexturePaint_Commit", "Commit Texture Paint"));
-
+		
+		Modify();
 	//	GWarn->BeginSlowTask(LOCTEXT("BeginMeshPaintMode_TexturePaint_CommitTask", "Committing Texture Paint Changes"), true);
 
 		int32 CurStep = 1;
diff --git a/Engine/Plugins/MeshPainting/Source/MeshPaintingToolset/Public/MeshPaintingToolsetTypes.h b/Engine/Plugins/MeshPainting/Source/MeshPaintingToolset/Public/MeshPaintingToolsetTypes.h
--- a/Engine/Plugins/MeshPainting/Source/MeshPaintingToolset/Public/MeshPaintingToolsetTypes.h
+++ b/Engine/Plugins/MeshPainting/Source/MeshPaintingToolset/Public/MeshPaintingToolsetTypes.h
@@ -124,6 +124,7 @@ struct FPaintTexture2DData
 	UPROPERTY(Transient)
 	TObjectPtr<UTexture2D> PaintingTexture2D = nullptr;
 
+	UPROPERTY(Transient)
 	bool bIsPaintingTexture2DModified = false;
 
 	/** Texture used to store the paint modifications for the transactions */
