Compile fixes for IWYU and upcoming dependency changes.

diff --git a/Engine/Plugins/Developer/Concert/ConcertSync/ConcertSyncServer/Source/ConcertSyncServer/ConcertSyncServer.Build.cs b/Engine/Plugins/Developer/Concert/ConcertSync/ConcertSyncServer/Source/ConcertSyncServer/ConcertSyncServer.Build.cs
--- a/Engine/Plugins/Developer/Concert/ConcertSync/ConcertSyncServer/Source/ConcertSyncServer/ConcertSyncServer.Build.cs
+++ b/Engine/Plugins/Developer/Concert/ConcertSync/ConcertSyncServer/Source/ConcertSyncServer/ConcertSyncServer.Build.cs
@@ -9,6 +9,9 @@ namespace UnrealBuildTool.Rules
 			// Make the linking path shorter (to prevent busting the Windows limit) when linking ConcertSyncServer.lib against an executable that have a long name.
 			ShortName = "CncrtSyncSvr";
 
+			// ConcertSyncServerLoop.inl depends on LaunchEngineLoop.h, but is never included in this module.
+			IWYUSupport = IWYUSupport.None;
+
 			PublicDependencyModuleNames.AddRange(
 				new string[]
 				{
diff --git a/Engine/Plugins/MeshPainting/Source/MeshPaintingToolset/Private/TexturePaintToolset.cpp b/Engine/Plugins/MeshPainting/Source/MeshPaintingToolset/Private/TexturePaintToolset.cpp
--- a/Engine/Plugins/MeshPainting/Source/MeshPaintingToolset/Private/TexturePaintToolset.cpp
+++ b/Engine/Plugins/MeshPainting/Source/MeshPaintingToolset/Private/TexturePaintToolset.cpp
@@ -18,6 +18,7 @@
 #include "MaterialShared.h"
 #include "MeshPaintingToolsetTypes.h"
 #include "TextureResource.h"
+#include "PixelFormat.h"
 
 #include UE_INLINE_GENERATED_CPP_BY_NAME(TexturePaintToolset)
 
@@ -317,6 +318,16 @@ bool UTexturePaintToolset::GenerateSeamMask(UMeshComponent* MeshComponent, int32
 	return RetVal;
 }
 
+int32 UTexturePaintToolset::GetMaxSupportedBytesPerPixelForPainting()
+{
+	return GPixelFormats[GetTempUncompressedTexturePixelFormat()].BlockBytes;
+}
+
+EPixelFormat UTexturePaintToolset::GetTempUncompressedTexturePixelFormat()
+{
+	return EPixelFormat::PF_B8G8R8A8;
+}
+
 UTexture2D* UTexturePaintToolset::CreateScratchUncompressedTexture(UTexture2D* SourceTexture)
 {
 	check(SourceTexture->Source.IsValid());
diff --git a/Engine/Plugins/MeshPainting/Source/MeshPaintingToolset/Public/TexturePaintToolset.h b/Engine/Plugins/MeshPainting/Source/MeshPaintingToolset/Public/TexturePaintToolset.h
--- a/Engine/Plugins/MeshPainting/Source/MeshPaintingToolset/Public/TexturePaintToolset.h
+++ b/Engine/Plugins/MeshPainting/Source/MeshPaintingToolset/Public/TexturePaintToolset.h
@@ -72,16 +72,10 @@ public:
 	static bool GenerateSeamMask(UMeshComponent* MeshComponent, int32 UVSet, UTextureRenderTarget2D* SeamRenderTexture, UTexture2D* Texture, UTextureRenderTarget2D* RenderTargetTexture );
 
 	/** Returns the maximum bytes per pixel that are supported for source textures when painting. This limitation is set by CreateTempUncompressedTexture. */
-	static constexpr int32 GetMaxSupportedBytesPerPixelForPainting()
-	{
-		return GPixelFormats[GetTempUncompressedTexturePixelFormat()].BlockBytes;
-	}
+	static int32 GetMaxSupportedBytesPerPixelForPainting();
 
 	/** Returns the pixel format that CreateTempUncompressedTexture uses to create render target data for painting. */
-	static constexpr EPixelFormat GetTempUncompressedTexturePixelFormat()
-	{
-		return EPixelFormat::PF_B8G8R8A8;
-	}
+	static EPixelFormat GetTempUncompressedTexturePixelFormat();
 
 	/** Static: Creates a temporary texture used to transfer data to a render target in memory */
 	static UTexture2D* CreateScratchUncompressedTexture(UTexture2D* SourceTexture);
diff --git a/Engine/Plugins/VirtualProduction/RemoteControl/Source/RemoteControl/Private/RemoteControlModule.h b/Engine/Plugins/VirtualProduction/RemoteControl/Source/RemoteControl/Private/RemoteControlModule.h
--- a/Engine/Plugins/VirtualProduction/RemoteControl/Source/RemoteControl/Private/RemoteControlModule.h
+++ b/Engine/Plugins/VirtualProduction/RemoteControl/Source/RemoteControl/Private/RemoteControlModule.h
@@ -6,8 +6,7 @@
 #include "CoreMinimal.h"
 #include "Engine/EngineTypes.h"
 #include "Factories/IRemoteControlMaskingFactory.h"
-
-class IRemoteControlInterceptionFeatureProcessor;
+#include "IRemoteControlInterceptionFeature.h"
 
 /**
  * Implementation of the RemoteControl interface
diff --git a/Engine/Source/Editor/MeshPaint/Private/TexturePaintHelpers.cpp b/Engine/Source/Editor/MeshPaint/Private/TexturePaintHelpers.cpp
--- a/Engine/Source/Editor/MeshPaint/Private/TexturePaintHelpers.cpp
+++ b/Engine/Source/Editor/MeshPaint/Private/TexturePaintHelpers.cpp
@@ -18,6 +18,7 @@
 #include "CanvasTypes.h"
 #include "CanvasItem.h"
 #include "MeshPaintTypes.h"
+#include "PixelFormat.h"
 
 void TexturePaintHelpers::CopyTextureToRenderTargetTexture(UTexture* SourceTexture, UTextureRenderTarget2D* RenderTargetTexture, ERHIFeatureLevel::Type FeatureLevel)
 {
@@ -315,6 +316,16 @@ bool TexturePaintHelpers::GenerateSeamMask(UMeshComponent* MeshComponent, int32
 	return RetVal;
 }
 
+int32 TexturePaintHelpers::GetMaxSupportedBytesPerPixelForPainting()
+{
+	return GPixelFormats[GetTempUncompressedTexturePixelFormat()].BlockBytes;
+}
+
+EPixelFormat TexturePaintHelpers::GetTempUncompressedTexturePixelFormat()
+{
+	return EPixelFormat::PF_B8G8R8A8;
+}
+
 UTexture2D* TexturePaintHelpers::CreateTempUncompressedTexture(UTexture2D* SourceTexture)
 {
 	check(SourceTexture->Source.IsValid());
diff --git a/Engine/Source/Editor/MeshPaint/Public/TexturePaintHelpers.h b/Engine/Source/Editor/MeshPaint/Public/TexturePaintHelpers.h
--- a/Engine/Source/Editor/MeshPaint/Public/TexturePaintHelpers.h
+++ b/Engine/Source/Editor/MeshPaint/Public/TexturePaintHelpers.h
@@ -66,16 +66,10 @@ public:
 	static bool GenerateSeamMask(UMeshComponent* MeshComponent, int32 UVSet, UTextureRenderTarget2D* SeamRenderTexture, UTexture2D* Texture, UTextureRenderTarget2D* RenderTargetTexture );
 
 	/** Returns the maximum bytes per pixel that are supported for source textures when painting. This limitation is set by CreateTempUncompressedTexture. */
-	static constexpr int32 GetMaxSupportedBytesPerPixelForPainting()
-	{
-		return GPixelFormats[GetTempUncompressedTexturePixelFormat()].BlockBytes;
-	}
+	static int32 GetMaxSupportedBytesPerPixelForPainting();
 
 	/** Returns the pixel format that CreateTempUncompressedTexture uses to create render target data for painting. */
-	static constexpr EPixelFormat GetTempUncompressedTexturePixelFormat()
-	{
-		return EPixelFormat::PF_B8G8R8A8;
-	}
+	static EPixelFormat GetTempUncompressedTexturePixelFormat();
 
 	/** Static: Creates a temporary texture used to transfer data to a render target in memory */
 	static UTexture2D* CreateTempUncompressedTexture(UTexture2D* SourceTexture);
diff --git a/Engine/Source/Runtime/Landscape/Public/LandscapeLight.h b/Engine/Source/Runtime/Landscape/Public/LandscapeLight.h
--- a/Engine/Source/Runtime/Landscape/Public/LandscapeLight.h
+++ b/Engine/Source/Runtime/Landscape/Public/LandscapeLight.h
@@ -7,6 +7,7 @@
 #pragma once
 
 #include "CoreMinimal.h"
+#include "PixelFormat.h"
 #include "StaticLighting.h"
 #include "RenderUtils.h"
 
