Fix up deprecated APIs

diff --git a/Engine/Plugins/MeshPainting/Source/MeshPaintingToolset/Private/MeshTexturePaintingTool.cpp b/Engine/Plugins/MeshPainting/Source/MeshPaintingToolset/Private/MeshTexturePaintingTool.cpp
--- a/Engine/Plugins/MeshPainting/Source/MeshPaintingToolset/Private/MeshTexturePaintingTool.cpp
+++ b/Engine/Plugins/MeshPainting/Source/MeshPaintingToolset/Private/MeshTexturePaintingTool.cpp
@@ -1074,7 +1074,7 @@ void UMeshTexturePaintingTool::PaintTexture(FMeshPaintParameters& InParams, TArr
 	}
 
 	check(GEditor && GEditor->GetEditorWorldContext().World());
-	const auto FeatureLevel = GEditor->GetEditorWorldContext().World()->FeatureLevel;
+	const auto FeatureLevel = GEditor->GetEditorWorldContext().World()->GetFeatureLevel();
 
 
 	FPaintTexture2DData* TextureData = GetPaintTargetData(PaintingTexture2D);
@@ -1706,7 +1706,7 @@ void UMeshTexturePaintingTool::ClearAllTextureOverrides()
 	if (UMeshPaintingSubsystem* MeshPaintingSubsystem = GEngine->GetEngineSubsystem<UMeshPaintingSubsystem>())
 	{
 		check(GEditor && GEditor->GetEditorWorldContext().World());
-		const auto FeatureLevel = GEditor->GetEditorWorldContext().World()->FeatureLevel;
+		const auto FeatureLevel = GEditor->GetEditorWorldContext().World()->GetFeatureLevel();
 		/** Remove all texture overrides which are currently stored and active */
 		for (decltype(PaintTargetData)::TIterator It(PaintTargetData); It; ++It)
 		{
diff --git a/Engine/Plugins/MeshPainting/Source/MeshPaintingToolset/Private/TexturePaintToolset.cpp b/Engine/Plugins/MeshPainting/Source/MeshPaintingToolset/Private/TexturePaintToolset.cpp
--- a/Engine/Plugins/MeshPainting/Source/MeshPaintingToolset/Private/TexturePaintToolset.cpp
+++ b/Engine/Plugins/MeshPainting/Source/MeshPaintingToolset/Private/TexturePaintToolset.cpp
@@ -393,7 +393,7 @@ void UTexturePaintToolset::SetupInitialRenderTargetData(UTexture2D* InTextureSou
 			check(TempSourceArtTexture != nullptr);
 
 			// Copy the texture to the render target using the GPU
-			CopyTextureToRenderTargetTexture(TempSourceArtTexture, InRenderTarget, GEditor->GetEditorWorldContext().World()->FeatureLevel);
+			CopyTextureToRenderTargetTexture(TempSourceArtTexture, InRenderTarget, GEditor->GetEditorWorldContext().World()->GetFeatureLevel());
 
 			// NOTE: TempSourceArtTexture is no longer needed (will be GC'd)
 		}
@@ -403,7 +403,7 @@ void UTexturePaintToolset::SetupInitialRenderTargetData(UTexture2D* InTextureSou
 		// Just copy (render) the texture in GPU memory to our render target.  Hopefully it's not
 		// compressed already!
 		check(InTextureSource->IsFullyStreamedIn());
-		CopyTextureToRenderTargetTexture(InTextureSource, InRenderTarget, GEditor->GetEditorWorldContext().World()->FeatureLevel);
+		CopyTextureToRenderTargetTexture(InTextureSource, InRenderTarget, GEditor->GetEditorWorldContext().World()->GetFeatureLevel());
 	}
 }
 
