Fix static analysis warning in MeshPaintModeHelpers.cpp

diff --git a/Engine/Plugins/MeshPainting/Source/MeshPaintEditorMode/Private/MeshPaintModeHelpers.cpp b/Engine/Plugins/MeshPainting/Source/MeshPaintEditorMode/Private/MeshPaintModeHelpers.cpp
--- a/Engine/Plugins/MeshPainting/Source/MeshPaintEditorMode/Private/MeshPaintModeHelpers.cpp
+++ b/Engine/Plugins/MeshPainting/Source/MeshPaintEditorMode/Private/MeshPaintModeHelpers.cpp
@@ -128,7 +128,7 @@ void UMeshPaintModeSubsystem::SetViewportColorMode(EMeshPaintActiveMode ActiveMo
 						UVChannel = Settings->UVChannel;
 					}
 
-					const UMeshComponent* LastPaintedComponent = MeshPaintingSubsystem->LastPaintedComponent;
+					const UMeshComponent* LastPaintedComponent = MeshPaintingSubsystem ? MeshPaintingSubsystem->LastPaintedComponent : nullptr;
 					if (LastPaintedComponent)
 					{
 						GVertexViewModeOverrideOwnerName = *LastPaintedComponent->GetOwner()->GetName();
