Fix non-unity
#rnx

diff --git a/Engine/Plugins/MeshPainting/Source/MeshPaintingToolset/Private/MeshTexturePaintingTool.cpp b/Engine/Plugins/MeshPainting/Source/MeshPaintingToolset/Private/MeshTexturePaintingTool.cpp
--- a/Engine/Plugins/MeshPainting/Source/MeshPaintingToolset/Private/MeshTexturePaintingTool.cpp
+++ b/Engine/Plugins/MeshPainting/Source/MeshPaintingToolset/Private/MeshTexturePaintingTool.cpp
@@ -25,6 +25,7 @@
 
 #include "Editor/EditorEngine.h"
 extern UNREALED_API class UEditorEngine* GEditor;
+extern ENGINE_API FString GVertexViewModeOverrideOwnerName;
 
 #include UE_INLINE_GENERATED_CPP_BY_NAME(MeshTexturePaintingTool)
 
