mesh paint tools: Fix broken keyboard shortcuts

diff --git a/Engine/Plugins/MeshPainting/Source/MeshPaintEditorMode/Private/MeshPaintMode.cpp b/Engine/Plugins/MeshPainting/Source/MeshPaintEditorMode/Private/MeshPaintMode.cpp
--- a/Engine/Plugins/MeshPainting/Source/MeshPaintEditorMode/Private/MeshPaintMode.cpp
+++ b/Engine/Plugins/MeshPainting/Source/MeshPaintEditorMode/Private/MeshPaintMode.cpp
@@ -362,7 +362,7 @@ void UMeshPaintMode::BindCommands()
 
 void UMeshPaintMode::OnToolStarted(UInteractiveToolManager* Manager, UInteractiveTool* Tool)
 {
-	FMeshPaintingToolActionCommands::UpdateToolCommandBinding(Tool, ToolCommandList, false);
+	FMeshPaintingToolActionCommands::UpdateToolCommandBinding(Tool, Toolkit->GetToolkitCommands(), false);
 
 	if (UMeshVertexPaintingTool* VertexPaintingTool = Cast<UMeshVertexPaintingTool>(GetToolManager()->GetActiveTool(EToolSide::Left)))
 	{
@@ -391,7 +391,7 @@ void UMeshPaintMode::OnVertexPaintFinished()
 
 void UMeshPaintMode::OnToolEnded(UInteractiveToolManager* Manager, UInteractiveTool* Tool)
 {
-	FMeshPaintingToolActionCommands::UpdateToolCommandBinding(Tool, ToolCommandList, true);
+	FMeshPaintingToolActionCommands::UpdateToolCommandBinding(Tool, Toolkit->GetToolkitCommands(), true);
 	// First update your bindings, then call the base behavior
 	Super::OnToolEnded(Manager, Tool);
 }
