Remove impossible Cast to UInterchangeGenericAssetsPipeline from UInterchangePythonPipelineBase
#rnx

diff --git a/Engine/Plugins/MeshPainting/Source/MeshPaintEditorMode/Private/MeshPaintModeHelpers.cpp b/Engine/Plugins/MeshPainting/Source/MeshPaintEditorMode/Private/MeshPaintModeHelpers.cpp
--- a/Engine/Plugins/MeshPainting/Source/MeshPaintEditorMode/Private/MeshPaintModeHelpers.cpp
+++ b/Engine/Plugins/MeshPainting/Source/MeshPaintEditorMode/Private/MeshPaintModeHelpers.cpp
@@ -344,14 +344,6 @@ void UMeshPaintModeSubsystem::ImportVertexColorsToSkeletalMesh(USkeletalMesh* Sk
 			{
 				UInterchangeGenericAssetsPipeline* GenericAssetPipeline = Cast<UInterchangeGenericAssetsPipeline>(PipelineBase);
 
-				if (!GenericAssetPipeline)
-				{
-					if (UInterchangePythonPipelineAsset* PythonPipelineAsset = Cast<UInterchangePythonPipelineAsset>(PipelineBase))
-					{
-						GenericAssetPipeline = Cast<UInterchangeGenericAssetsPipeline>(PythonPipelineAsset->GeneratedPipeline);
-					}
-				}
-
 				if (GenericAssetPipeline)
 				{
 					if (GenericAssetPipeline->CommonMeshesProperties && GenericAssetPipeline->CommonMeshesProperties->VertexColorImportOption != EInterchangeVertexColorImportOption::IVCIO_Ignore)
diff --git a/Engine/Plugins/MeshPainting/Source/MeshPaintingToolset/Private/MeshPaintSkeletalMeshAdapter.cpp b/Engine/Plugins/MeshPainting/Source/MeshPaintingToolset/Private/MeshPaintSkeletalMeshAdapter.cpp
--- a/Engine/Plugins/MeshPainting/Source/MeshPaintingToolset/Private/MeshPaintSkeletalMeshAdapter.cpp
+++ b/Engine/Plugins/MeshPainting/Source/MeshPaintingToolset/Private/MeshPaintSkeletalMeshAdapter.cpp
@@ -403,14 +403,6 @@ void FMeshPaintSkeletalMeshComponentAdapter::PreEdit()
 			{
 				UInterchangeGenericAssetsPipeline* GenericAssetPipeline = Cast<UInterchangeGenericAssetsPipeline>(PipelineBase);
 
-				if (!GenericAssetPipeline)
-				{
-					if (UInterchangePythonPipelineAsset* PythonPipelineAsset = Cast<UInterchangePythonPipelineAsset>(PipelineBase))
-					{
-						GenericAssetPipeline = Cast<UInterchangeGenericAssetsPipeline>(PythonPipelineAsset->GeneratedPipeline);
-					}
-				}
-				
 				if (GenericAssetPipeline)
 				{
 					if (GenericAssetPipeline->CommonMeshesProperties && GenericAssetPipeline->CommonMeshesProperties->VertexColorImportOption != EInterchangeVertexColorImportOption::IVCIO_Ignore)
