Fixed an issue with the mesh painting tools loosing its texture override on an material intstance when an static permutation is enbabled.
The override would be not be removed on the old material ressource and wasn't applied to the new material ressource.

The solution chosen here was to add an global state that keep track of the override(s) applied via the mesh paint adaters. When an override is active the global state watch for any material modification.
When an overriden materail is modifed before the modification we turn off the texture overrides to clean the existing material ressource used and then after the modification we apply back the override to the material thus overriding the new material ressource.

diff --git a/Engine/Plugins/MeshPainting/Source/MeshPaintingToolset/Private/IMeshPaintComponentAdapter.cpp b/Engine/Plugins/MeshPainting/Source/MeshPaintingToolset/Private/IMeshPaintComponentAdapter.cpp
--- a/Engine/Plugins/MeshPainting/Source/MeshPaintingToolset/Private/IMeshPaintComponentAdapter.cpp
+++ b/Engine/Plugins/MeshPainting/Source/MeshPaintingToolset/Private/IMeshPaintComponentAdapter.cpp
@@ -2,22 +2,28 @@
 
 #include "IMeshPaintComponentAdapter.h"
 
+#include "Components/MeshComponent.h"
+#include "Containers/ContainerAllocationPolicies.h"
 #include "Engine/Texture.h"
 #include "Engine/World.h"
-#include "Components/MeshComponent.h"
-#include "MeshPaintingToolsetTypes.h"
+#include "MaterialShared.h"
 #include "Materials/Material.h"
-#include "Materials/MaterialInstance.h"
 #include "Materials/MaterialExpressionTextureCoordinate.h"
 #include "Materials/MaterialExpressionTextureSampleParameter.h"
-#include "MaterialShared.h"
+#include "Materials/MaterialInstance.h"
+#include "MeshPaintingToolsetTypes.h"
 #include "TexturePaintToolset.h"
+#include "UObject/UObjectGlobals.h"
+#include "UObject/UnrealType.h"
+#include "UObject/WeakObjectPtrTemplates.h"
 
 //////////////////////////////////////////////////////////////////////////
 // IMeshPaintGeometryAdapter
 
 void IMeshPaintComponentAdapter::DefaultApplyOrRemoveTextureOverride(UMeshComponent* InMeshComponent, UTexture* SourceTexture, UTexture* OverrideTexture)
 {
+	// Legacy implementation for the deprecation duration. For a valid implementation look at deprecation comment
+
 	const ERHIFeatureLevel::Type FeatureLevel = InMeshComponent->GetWorld()->GetFeatureLevel();
 
 	// Check all the materials on the mesh to see if the user texture is there
@@ -25,8 +31,7 @@ void IMeshPaintComponentAdapter::DefaultApplyOrRemoveTextureOverride(UMeshCompon
 	UMaterialInterface* MaterialToCheck = InMeshComponent->GetMaterial(MaterialIndex);
 	while (MaterialToCheck != nullptr)
 	{
-		const bool bIsTextureUsed = DoesMaterialUseTexture(MaterialToCheck, SourceTexture);
-		if (bIsTextureUsed)
+		if (!OverrideTexture || DoesMaterialUseTexture(MaterialToCheck, SourceTexture))
 		{
 			MaterialToCheck->OverrideTexture(SourceTexture, OverrideTexture, FeatureLevel);
 		}
@@ -133,3 +138,299 @@ void IMeshPaintComponentAdapter::DefaultQueryPaintableTextures(int32 MaterialInd
 		}
 	}
 }
+
+namespace UE::MeshPaintingToolset
+{
+	namespace Private
+	{
+		struct FOverrideData
+		{
+			TWeakObjectPtr<UTexture> OverrideTexture;
+			TSet<ERHIFeatureLevel::Type, DefaultKeyFuncs<ERHIFeatureLevel::Type>, TInlineSetAllocator<1>> OverridenFeatureLevels;
+			uint32 Count = 0;
+		};
+
+		struct FGlobalTextureOverrideState
+		{
+			using FOverrideKey = TPair<TWeakObjectPtr<UMaterialInterface>, TWeakObjectPtr<const UTexture>>;
+
+			static void DuplicateOverrideOwnership(const FDefaultTextureOverride* Current, const FDefaultTextureOverride* To)
+			{
+				TSet<FOverrideKey, DefaultKeyFuncs<FOverrideKey>, TInlineSetAllocator<2>>* Overrides = DefaultTextureOverrideToOverrides.Find(Current);
+				if (Overrides)
+				{
+					for (const FOverrideKey& Pair : *Overrides)
+					{
+						OverridesData.Find(Pair)->Count += 1;
+					}
+
+					// Duplicate the override in case an allocation change the addresses
+					DefaultTextureOverrideToOverrides.Add(To, TSet<FOverrideKey>(*Overrides));
+				}
+			}
+
+			static void TransferOverrideOwnership(const FDefaultTextureOverride* Current, const FDefaultTextureOverride* To)
+			{
+				TSet<FOverrideKey, DefaultKeyFuncs<FOverrideKey>, TInlineSetAllocator<2>> Overrides;
+				if (DefaultTextureOverrideToOverrides.RemoveAndCopyValue(Current, Overrides))
+				{
+					DefaultTextureOverrideToOverrides.Add(To, MoveTemp(Overrides));
+				}
+			}
+
+			static void RegisterMaterialOverride(const FDefaultTextureOverride* Requester, UMaterialInterface* Material, const UTexture* SourceTexture, UTexture* OverrrideTexture, const ERHIFeatureLevel::Type FeatureLevel)
+			{
+				TSet<FOverrideKey, DefaultKeyFuncs<FOverrideKey>, TInlineSetAllocator<2>>& Overrides = DefaultTextureOverrideToOverrides.FindOrAdd(Requester);
+				
+				FOverrideKey Pair = FOverrideKey(TWeakObjectPtr<UMaterialInterface>(Material), TWeakObjectPtr<const UTexture>(SourceTexture));
+				uint32 PairHash = GetTypeHash(Pair);
+
+				Overrides.AddByHash(PairHash, Pair);
+				
+				FOverrideData& OverrideData = OverridesData.FindOrAddByHash(PairHash, Pair);
+				
+				bool bFeatureLevelAlreadyOverriden = false;
+				uint32 FeatureLevelHash = GetTypeHash(FeatureLevel);
+				OverrideData.OverridenFeatureLevels.AddByHash(FeatureLevelHash, FeatureLevel, &bFeatureLevelAlreadyOverriden);
+
+				if (OverrideData.Count == 0 || OverrideData.OverrideTexture != OverrrideTexture)
+				{
+					OverrideData.OverrideTexture = OverrrideTexture;
+
+					for (const ERHIFeatureLevel::Type LevelToUpdate : OverrideData.OverridenFeatureLevels)
+					{
+						Material->OverrideTexture(SourceTexture, OverrrideTexture, LevelToUpdate);
+					}
+
+					AddMaterialTracking(Pair);
+				}
+				else if (!bFeatureLevelAlreadyOverriden)
+				{
+					OverrideData.OverridenFeatureLevels.AddByHash(FeatureLevelHash, FeatureLevel);
+					Material->OverrideTexture(SourceTexture, OverrrideTexture, FeatureLevel);
+				}
+	
+				OverrideData.Count += 1;
+			}
+
+			static void RemoveMaterialOverride(const FDefaultTextureOverride* Requester, UMaterialInterface* Material, const UTexture* SourceTexture)
+			{
+				if (TSet<FOverrideKey, DefaultKeyFuncs<FOverrideKey>, TInlineSetAllocator<2>>* Override = DefaultTextureOverrideToOverrides.Find(Requester))
+				{
+					FOverrideKey Pair = FOverrideKey(TWeakObjectPtr<UMaterialInterface>(Material), TWeakObjectPtr<const UTexture>(SourceTexture));
+					uint32 PairHash = GetTypeHash(Pair);
+
+					if (Override->RemoveByHash(PairHash, Pair) > 0)
+					{
+						if (FOverrideData* OverrideData = OverridesData.FindByHash(PairHash, Pair))
+						{
+							OverrideData->Count -= 1;
+							if (OverrideData->Count == 0)
+							{
+								for (const ERHIFeatureLevel::Type FeatureLevel : OverrideData->OverridenFeatureLevels)
+								{
+									Material->OverrideTexture(SourceTexture, nullptr, FeatureLevel);
+								}
+
+								OverridesData.RemoveByHash(PairHash, Pair);
+								RemoveMaterialTracking(Pair);
+							}
+						}
+					}
+				}
+			}
+
+			static void FreeOverrideOwnership(const FDefaultTextureOverride* Requester)
+			{
+				TSet<FOverrideKey, DefaultKeyFuncs<FOverrideKey>, TInlineSetAllocator<2>> Overrides;
+				if (DefaultTextureOverrideToOverrides.RemoveAndCopyValue(Requester, Overrides))
+				{
+					for (const FOverrideKey& Override : Overrides)
+					{
+						uint32 OverrideHash = GetTypeHash(Override);
+						if (FOverrideData* OverrideData = OverridesData.FindByHash(OverrideHash, Override))
+						{
+							OverrideData->Count -= 1;
+							if (OverrideData->Count == 0)
+							{
+								UMaterialInterface* Material = Override.Key.Get();
+								const UTexture* SourceTexture = Override.Value.Get();
+								for (const ERHIFeatureLevel::Type LevelToUpdate : OverrideData->OverridenFeatureLevels)
+								{
+									Material->OverrideTexture(SourceTexture, nullptr, LevelToUpdate);
+								}
+
+								OverridesData.RemoveByHash(OverrideHash, Override);
+								RemoveMaterialTracking(Override);
+							}
+						}
+					}
+				}
+			}
+
+			static void AddMaterialTracking(const FOverrideKey& Override)
+			{
+				if (MaterialsAndTexturesOverriden.IsEmpty())
+				{
+					OnObjectModifiedDelegateHandle = FCoreUObjectDelegates::OnObjectModified.AddStatic(&FGlobalTextureOverrideState::OnObjectModified);
+					PostEditDelegateHandle = FCoreUObjectDelegates::OnObjectPropertyChanged.AddStatic(&FGlobalTextureOverrideState::OnObjectPropertyChanged);
+				}
+
+				TSet<TWeakObjectPtr<const UTexture>, DefaultKeyFuncs<TWeakObjectPtr<const UTexture>>, TInlineSetAllocator<2>>& Textures = MaterialsAndTexturesOverriden.FindOrAdd(Override.Key);
+				Textures.Add(Override.Value);
+			}
+
+			static void RemoveMaterialTracking(const FOverrideKey& Override)
+			{
+				uint32 MaterialHash = GetTypeHash(Override.Key);
+				if (TSet<TWeakObjectPtr<const UTexture>, DefaultKeyFuncs<TWeakObjectPtr<const UTexture>>, TInlineSetAllocator<2>>* Textures = MaterialsAndTexturesOverriden.FindByHash(MaterialHash, Override.Key))
+				{
+					Textures->Remove(Override.Value);
+					if (Textures->IsEmpty())
+					{
+						MaterialsAndTexturesOverriden.RemoveByHash(MaterialHash, Override.Key);
+						if (MaterialsAndTexturesOverriden.IsEmpty())
+						{
+							FCoreUObjectDelegates::OnObjectModified.Remove(OnObjectModifiedDelegateHandle);
+							FCoreUObjectDelegates::OnObjectPropertyChanged.Remove(PostEditDelegateHandle);
+						}
+					}
+				}
+			}
+
+			static void OnObjectModified(UObject* Object)
+			{
+				if (UMaterialInterface* Material = Cast<UMaterialInterface>(Object))
+				{
+					FOverrideKey Override;
+					Override.Key = Material;
+					if (const TSet<TWeakObjectPtr<const UTexture>, DefaultKeyFuncs<TWeakObjectPtr<const UTexture>>, TInlineSetAllocator<2>>* Textures = MaterialsAndTexturesOverriden.Find(Override.Key))
+					{
+						for (const TWeakObjectPtr<const UTexture>& Texture : *Textures)
+						{
+							if (const UTexture* RawTexturePtr = Texture.Get())
+							{
+								Override.Value = Texture;
+								if (FOverrideData* OverrideData = OverridesData.Find(Override))
+								{
+									for (const ERHIFeatureLevel::Type FeatureLevel : OverrideData->OverridenFeatureLevels)
+									{
+										// The material resource might change because of the modifications. To avoid leaking some temp texture overrides, this just remove the temporary overrides during the modification.
+										Material->OverrideTexture(RawTexturePtr, nullptr, FeatureLevel);
+									}
+								}
+							}
+						}
+					}
+				}
+			}
+
+			static void OnObjectPropertyChanged(UObject* Object, FPropertyChangedEvent& PropertyChangedEvent)
+			{
+				if (PropertyChangedEvent.ChangeType != EPropertyChangeType::Interactive)
+				{
+					if (UMaterialInterface* Material = Cast<UMaterialInterface>(Object))
+					{
+						FOverrideKey Override;
+						Override.Key = Material;
+						if (const TSet<TWeakObjectPtr<const UTexture>, DefaultKeyFuncs<TWeakObjectPtr<const UTexture>>, TInlineSetAllocator<2>>* Textures = MaterialsAndTexturesOverriden.Find(Override.Key))
+						{
+							for (const TWeakObjectPtr<const UTexture>& Texture : *Textures)
+							{
+								if (const UTexture* RawTexturePtr = Texture.Get())
+								{
+									Override.Value = Texture;
+									if (FOverrideData* OverrideData = OverridesData.Find(Override))
+									{
+										if (UTexture* OverrideTexture = OverrideData->OverrideTexture.Get())
+										{ 
+											for (const ERHIFeatureLevel::Type FeatureLevel : OverrideData->OverridenFeatureLevels)
+											{
+												// Reapply the temporary overrides after the modification.
+												Material->OverrideTexture(RawTexturePtr, OverrideTexture, FeatureLevel);
+											}
+										}
+									}
+								}
+							}
+						}
+					}
+				}
+			}
+
+
+		private:
+			static TMap<FOverrideKey, FOverrideData> OverridesData;
+			static TMap<const FDefaultTextureOverride*, TSet<FOverrideKey, DefaultKeyFuncs<FOverrideKey>, TInlineSetAllocator<2>>> DefaultTextureOverrideToOverrides;
+			static TMap<TWeakObjectPtr<UMaterialInterface>, TSet<TWeakObjectPtr<const UTexture>, DefaultKeyFuncs<TWeakObjectPtr<const UTexture>>, TInlineSetAllocator<2>>> MaterialsAndTexturesOverriden;
+
+			static FDelegateHandle OnObjectModifiedDelegateHandle;
+			static FDelegateHandle PostEditDelegateHandle;
+		};
+
+		TMap<FGlobalTextureOverrideState::FOverrideKey, FOverrideData> FGlobalTextureOverrideState::OverridesData;
+		TMap<const FDefaultTextureOverride*, TSet<FGlobalTextureOverrideState::FOverrideKey, DefaultKeyFuncs<FGlobalTextureOverrideState::FOverrideKey>, TInlineSetAllocator<2>>> FGlobalTextureOverrideState::DefaultTextureOverrideToOverrides;
+		TMap<TWeakObjectPtr<UMaterialInterface>, TSet<TWeakObjectPtr<const UTexture>, DefaultKeyFuncs<TWeakObjectPtr<const UTexture>>, TInlineSetAllocator<2>>> FGlobalTextureOverrideState::MaterialsAndTexturesOverriden;
+		FDelegateHandle FGlobalTextureOverrideState::OnObjectModifiedDelegateHandle;
+		FDelegateHandle FGlobalTextureOverrideState::PostEditDelegateHandle;
+	}
+
+
+	FDefaultTextureOverride::FDefaultTextureOverride(const FDefaultTextureOverride& InOther)
+	{
+		operator=(InOther);
+	}
+
+	FDefaultTextureOverride::FDefaultTextureOverride(FDefaultTextureOverride&& InOther)
+	{
+		operator=(MoveTemp(InOther));
+	}
+
+	FDefaultTextureOverride& FDefaultTextureOverride::operator=(const FDefaultTextureOverride& InOther)
+	{
+		Private::FGlobalTextureOverrideState::DuplicateOverrideOwnership(this, &InOther);
+		return *this;
+	}
+
+	FDefaultTextureOverride& FDefaultTextureOverride::operator=(FDefaultTextureOverride&& InOther)
+	{
+		Private::FGlobalTextureOverrideState::TransferOverrideOwnership(this, &InOther);
+		return *this;
+	}
+
+	void FDefaultTextureOverride::ApplyOrRemoveTextureOverride(const UMeshComponent* InMeshComponent, const UTexture* SourceTexture, UTexture* OverrideTexture) const
+	{
+		check(IsInGameThread());
+
+		const ERHIFeatureLevel::Type FeatureLevel = InMeshComponent->GetWorld()->GetFeatureLevel();
+
+		// Check all the materials on the mesh to see if the user texture is there
+		int32 MaterialIndex = 0;
+		UMaterialInterface* MaterialToCheck = InMeshComponent->GetMaterial(MaterialIndex);
+		while (MaterialToCheck != nullptr)
+		{
+			if (!OverrideTexture)
+			{
+				// Unregister for all materials. This will no affect the material that weren't overridden by this instance
+				Private::FGlobalTextureOverrideState::RemoveMaterialOverride(this, MaterialToCheck, SourceTexture);
+			}
+			else if (DoesMaterialUseTexture(MaterialToCheck, SourceTexture))
+			{
+				// Keep track of the material overridden
+				Private::FGlobalTextureOverrideState::RegisterMaterialOverride(this, MaterialToCheck, SourceTexture, OverrideTexture, FeatureLevel);
+
+			}
+
+			++MaterialIndex;
+			MaterialToCheck = InMeshComponent->GetMaterial(MaterialIndex);
+		}
+	}
+
+	FDefaultTextureOverride::~FDefaultTextureOverride()
+	{
+		check(IsInGameThread());
+
+		Private::FGlobalTextureOverrideState::FreeOverrideOwnership(this);
+	}
+
+}
diff --git a/Engine/Plugins/MeshPainting/Source/MeshPaintingToolset/Private/MeshPaintGeometryCollectionAdapter.cpp b/Engine/Plugins/MeshPainting/Source/MeshPaintingToolset/Private/MeshPaintGeometryCollectionAdapter.cpp
--- a/Engine/Plugins/MeshPainting/Source/MeshPaintingToolset/Private/MeshPaintGeometryCollectionAdapter.cpp
+++ b/Engine/Plugins/MeshPainting/Source/MeshPaintingToolset/Private/MeshPaintGeometryCollectionAdapter.cpp
@@ -243,7 +243,7 @@ void FMeshPaintGeometryCollectionComponentAdapter::ApplyOrRemoveTextureOverride(
 {
 	if (GeometryCollectionComponent.IsValid())
 	{
-		DefaultApplyOrRemoveTextureOverride(GeometryCollectionComponent.Get(), SourceTexture, OverrideTexture);
+		TextureOverridesState.ApplyOrRemoveTextureOverride(GeometryCollectionComponent.Get(), SourceTexture, OverrideTexture);
 	}
 }
 
diff --git a/Engine/Plugins/MeshPainting/Source/MeshPaintingToolset/Private/MeshPaintSkeletalMeshAdapter.cpp b/Engine/Plugins/MeshPainting/Source/MeshPaintingToolset/Private/MeshPaintSkeletalMeshAdapter.cpp
--- a/Engine/Plugins/MeshPainting/Source/MeshPaintingToolset/Private/MeshPaintSkeletalMeshAdapter.cpp
+++ b/Engine/Plugins/MeshPainting/Source/MeshPaintingToolset/Private/MeshPaintSkeletalMeshAdapter.cpp
@@ -343,7 +343,7 @@ void FMeshPaintSkeletalMeshComponentAdapter::ApplyOrRemoveTextureOverride(UTextu
 {
 	if (SkeletalMeshComponent.IsValid())
 	{
-		DefaultApplyOrRemoveTextureOverride(SkeletalMeshComponent.Get(), SourceTexture, OverrideTexture);
+		TextureOverridesState.ApplyOrRemoveTextureOverride(SkeletalMeshComponent.Get(), SourceTexture, OverrideTexture);
 	}
 }
 
diff --git a/Engine/Plugins/MeshPainting/Source/MeshPaintingToolset/Private/MeshPaintStaticMeshAdapter.cpp b/Engine/Plugins/MeshPainting/Source/MeshPaintingToolset/Private/MeshPaintStaticMeshAdapter.cpp
--- a/Engine/Plugins/MeshPainting/Source/MeshPaintingToolset/Private/MeshPaintStaticMeshAdapter.cpp
+++ b/Engine/Plugins/MeshPainting/Source/MeshPaintingToolset/Private/MeshPaintStaticMeshAdapter.cpp
@@ -253,7 +253,7 @@ void FMeshPaintStaticMeshComponentAdapter::ApplyOrRemoveTextureOverride(UTexture
 {
 	if (StaticMeshComponent.IsValid())
 	{
-		DefaultApplyOrRemoveTextureOverride(StaticMeshComponent.Get(), SourceTexture, OverrideTexture);
+		TextureOverridesState.ApplyOrRemoveTextureOverride(StaticMeshComponent.Get(), SourceTexture, OverrideTexture);
 	}
 }
 
diff --git a/Engine/Plugins/MeshPainting/Source/MeshPaintingToolset/Public/IMeshPaintComponentAdapter.h b/Engine/Plugins/MeshPainting/Source/MeshPaintingToolset/Public/IMeshPaintComponentAdapter.h
--- a/Engine/Plugins/MeshPainting/Source/MeshPaintingToolset/Public/IMeshPaintComponentAdapter.h
+++ b/Engine/Plugins/MeshPainting/Source/MeshPaintingToolset/Public/IMeshPaintComponentAdapter.h
@@ -7,6 +7,8 @@
 class FReferenceCollector;
 class UMeshComponent;
 class UTexture;
+class UMaterialInterface;
+
 namespace UE
 {
 	namespace Geometry
@@ -15,6 +17,8 @@ namespace UE
 	}
 }
 
+namespace ERHIFeatureLevel { enum Type : int; }
+
 /**
  * Interface for a class to provide mesh painting support for a subclass of UMeshComponent
  */
@@ -82,6 +86,7 @@ public:
 	virtual void PostEdit() = 0;
 
 	/** Default functionality for applying or remove a texture override */
+	UE_DEPRECATED(5.4, "This method will be removed because its implementation is not able to handle an material instance being edited. Please inherit from or use an UE::MeshPaintingToolset::FDefaultTextureOverride object and use it function ApplyOrRemoveTextureOverride.")
 	static void DefaultApplyOrRemoveTextureOverride(UMeshComponent* InMeshComponent, UTexture* SourceTexture, UTexture* OverrideTexture);
 
 	/** Default functionality for querying textures from a mesh component which are suitable for texture painting */
@@ -104,6 +109,29 @@ public:
 
 };
 
+namespace UE::MeshPaintingToolset
+{
+	/*
+	 * Provide a default implementation for the texture override(s) needed when painting a mesh component
+	 * This object will keep track of the overridden material resources and remove the override when they are no longer referred by an default texture override
+	 */
+	class MESHPAINTINGTOOLSET_API FDefaultTextureOverride final
+	{
+	public:
+		void ApplyOrRemoveTextureOverride(const UMeshComponent* InMeshComponent, const UTexture* SourceTexture, UTexture* OverrideTexture) const;
+
+		~FDefaultTextureOverride();
+
+		FDefaultTextureOverride() = default;
+
+		FDefaultTextureOverride(const FDefaultTextureOverride& InOther);
+		FDefaultTextureOverride(FDefaultTextureOverride&& InOther);
+
+		FDefaultTextureOverride& operator=(const FDefaultTextureOverride& InOther);
+		FDefaultTextureOverride& operator=(FDefaultTextureOverride&& InOther);
+	};
+}
+
 #if UE_ENABLE_INCLUDE_ORDER_DEPRECATED_IN_5_2
 #include "CoreMinimal.h"
 #endif
diff --git a/Engine/Plugins/MeshPainting/Source/MeshPaintingToolset/Public/MeshPaintGeometryCollectionAdapter.h b/Engine/Plugins/MeshPainting/Source/MeshPaintingToolset/Public/MeshPaintGeometryCollectionAdapter.h
--- a/Engine/Plugins/MeshPainting/Source/MeshPaintingToolset/Public/MeshPaintGeometryCollectionAdapter.h
+++ b/Engine/Plugins/MeshPainting/Source/MeshPaintingToolset/Public/MeshPaintGeometryCollectionAdapter.h
@@ -55,6 +55,9 @@ protected:
 
 	bool bSavedShowBoneColors = false;
 
+	// Texture override state
+	UE::MeshPaintingToolset::FDefaultTextureOverride TextureOverridesState;
+
 	/// Get the underlying UGeometryCollection from the component, as a non-const object
 	/// Caller must have already validated that the component weak pointer is still valid (as this is called per-vertex)
 	UGeometryCollection* GetGeometryCollectionObject() const;
diff --git a/Engine/Plugins/MeshPainting/Source/MeshPaintingToolset/Public/MeshPaintSkeletalMeshAdapter.h b/Engine/Plugins/MeshPainting/Source/MeshPaintingToolset/Public/MeshPaintSkeletalMeshAdapter.h
--- a/Engine/Plugins/MeshPainting/Source/MeshPaintingToolset/Public/MeshPaintSkeletalMeshAdapter.h
+++ b/Engine/Plugins/MeshPainting/Source/MeshPaintingToolset/Public/MeshPaintSkeletalMeshAdapter.h
@@ -69,6 +69,9 @@ protected:
 	FSkeletalMeshLODModel* LODModel;
 	/** LOD Index for which data has to be retrieved / altered*/
 	int32 MeshLODIndex;
+
+	// Texture override state
+	UE::MeshPaintingToolset::FDefaultTextureOverride TextureOverridesState;
 };
 
 //////////////////////////////////////////////////////////////////////////
diff --git a/Engine/Plugins/MeshPainting/Source/MeshPaintingToolset/Public/MeshPaintStaticMeshAdapter.h b/Engine/Plugins/MeshPainting/Source/MeshPaintingToolset/Public/MeshPaintStaticMeshAdapter.h
--- a/Engine/Plugins/MeshPainting/Source/MeshPaintingToolset/Public/MeshPaintStaticMeshAdapter.h
+++ b/Engine/Plugins/MeshPainting/Source/MeshPaintingToolset/Public/MeshPaintStaticMeshAdapter.h
@@ -59,6 +59,9 @@ protected:
 	FStaticMeshLODResources* LODModel;
 	/** LOD Index for which data has to be retrieved / altered*/
 	int32 MeshLODIndex;
+
+	// Texture override state
+	UE::MeshPaintingToolset::FDefaultTextureOverride TextureOverridesState;
 };
 
 //////////////////////////////////////////////////////////////////////////
diff --git a/Engine/Source/Editor/UnrealEd/Private/PreviewMaterial.cpp b/Engine/Source/Editor/UnrealEd/Private/PreviewMaterial.cpp
--- a/Engine/Source/Editor/UnrealEd/Private/PreviewMaterial.cpp
+++ b/Engine/Source/Editor/UnrealEd/Private/PreviewMaterial.cpp
@@ -472,6 +472,9 @@ void UMaterialEditorInstanceConstant::PostEditChangeProperty(FPropertyChangedEve
 {
 	if (SourceInstance)
 	{
+		// Warn our source instance that it is about to be updated.
+		SourceInstance->PreEditChange(PropertyChangedEvent.Property);
+
 		FProperty* PropertyThatChanged = PropertyChangedEvent.Property;
 		bool bLayersParameterChanged = false;
 
