Vertex Paint: Fixed broken logic to check for editable mesh data

diff --git a/Engine/Plugins/MeshPainting/Source/MeshPaintingToolset/Private/MeshPaintSkeletalMeshAdapter.cpp b/Engine/Plugins/MeshPainting/Source/MeshPaintingToolset/Private/MeshPaintSkeletalMeshAdapter.cpp
--- a/Engine/Plugins/MeshPainting/Source/MeshPaintingToolset/Private/MeshPaintSkeletalMeshAdapter.cpp
+++ b/Engine/Plugins/MeshPainting/Source/MeshPaintingToolset/Private/MeshPaintSkeletalMeshAdapter.cpp
@@ -27,9 +27,9 @@ void PropagateVertexPaintToSkeletalMesh(USkeletalMesh* SkeletalMesh, int32 LODIn
 		int32 SoftVertexIndexes[3];
 	};
 
-	if (!SkeletalMesh || SkeletalMesh->HasMeshDescription(LODIndex))
+	if (!SkeletalMesh || !SkeletalMesh->HasMeshDescription(LODIndex))
 	{
-		//We do not propagate vertex color for old asset
+		// We do not propagate vertex color for LODs that don't have editable mesh data.
 		return;
 	}
 
diff --git a/Engine/Source/Editor/MeshPaint/Private/MeshPaintSkeletalMeshAdapter.cpp b/Engine/Source/Editor/MeshPaint/Private/MeshPaintSkeletalMeshAdapter.cpp
--- a/Engine/Source/Editor/MeshPaint/Private/MeshPaintSkeletalMeshAdapter.cpp
+++ b/Engine/Source/Editor/MeshPaint/Private/MeshPaintSkeletalMeshAdapter.cpp
@@ -29,9 +29,9 @@ void PropagateVertexPaintToAsset(USkeletalMesh* SkeletalMesh, int32 LODIndex)
 		int32 SoftVertexIndexes[3];
 	};
 
-	if (!SkeletalMesh || SkeletalMesh->HasMeshDescription(LODIndex))
+	if (!SkeletalMesh || !SkeletalMesh->HasMeshDescription(LODIndex))
 	{
-		//We do not propagate vertex color for old asset
+		// We do not propagate vertex color for LODs that don't have editable mesh data.
 		return;
 	}
 
