Enable texture mesh painting to virtual textures
#tests Go to mesh painting mode and select a texture to paint which is a virutal texture.

diff --git a/Engine/Plugins/MeshPainting/Source/MeshPaintingToolset/Private/IMeshPaintComponentAdapter.cpp b/Engine/Plugins/MeshPainting/Source/MeshPaintingToolset/Private/IMeshPaintComponentAdapter.cpp
--- a/Engine/Plugins/MeshPainting/Source/MeshPaintingToolset/Private/IMeshPaintComponentAdapter.cpp
+++ b/Engine/Plugins/MeshPainting/Source/MeshPaintingToolset/Private/IMeshPaintComponentAdapter.cpp
@@ -45,7 +45,6 @@ static bool IsTextureSuitableForTexturePainting(const TWeakObjectPtr<UTexture> T
 {
 	return (TexturePtr.Get() != nullptr &&
 		!TexturePtr->IsNormalMap() &&
-		!TexturePtr->VirtualTextureStreaming &&
 		!TexturePtr->HasHDRSource() && // Currently HDR textures are not supported to paint on.
 		TexturePtr->Source.IsValid() &&
 		TexturePtr->Source.GetBytesPerPixel() > 0 && // Textures' sources must have a known count of bytes per pixel
diff --git a/Engine/Plugins/MeshPainting/Source/MeshPaintingToolset/Private/MeshTexturePaintingTool.cpp b/Engine/Plugins/MeshPainting/Source/MeshPaintingToolset/Private/MeshTexturePaintingTool.cpp
--- a/Engine/Plugins/MeshPainting/Source/MeshPaintingToolset/Private/MeshTexturePaintingTool.cpp
+++ b/Engine/Plugins/MeshPainting/Source/MeshPaintingToolset/Private/MeshTexturePaintingTool.cpp
@@ -21,6 +21,7 @@
 #include "TextureResource.h"
 #include "Editor/TransBuffer.h"
 #include "Editor/UnrealEdEngine.h"
+#include "VT/VirtualTextureAdapter.h"
 #include "UnrealEdGlobals.h"
 
 #include "Editor/EditorEngine.h"
@@ -687,6 +688,14 @@ bool UMeshTexturePaintingTool::PaintInternal(const TArrayView<TPair<FVector, FVe
 	return bPaintApplied;
 }
 
+/** Painting texture to use in material override should be the virtual texture adapter if it exists. */
+static UTexture* GetTextureForMaterialOverride(FPaintTexture2DData const& TextureData)
+{
+	UTexture* RenderTarget = TextureData.PaintRenderTargetTexture;
+	UTexture* RenderTargetAdapter = TextureData.PaintRenderTargetTextureAdapter;
+	return RenderTargetAdapter != nullptr ? RenderTargetAdapter : RenderTarget;
+}
+
 void UMeshTexturePaintingTool::AddTextureOverrideToComponent(FPaintTexture2DData& TextureData, UMeshComponent* MeshComponent, const IMeshPaintComponentAdapter* MeshPaintAdapter)
 {
 	if (MeshComponent && MeshPaintAdapter)
@@ -695,7 +704,7 @@ void UMeshTexturePaintingTool::AddTextureOverrideToComponent(FPaintTexture2DData
 		{
 			TextureData.PaintedComponents.AddUnique(MeshComponent);
 
-			MeshPaintAdapter->ApplyOrRemoveTextureOverride(TextureData.PaintingTexture2D, TextureData.PaintRenderTargetTexture);
+			MeshPaintAdapter->ApplyOrRemoveTextureOverride(TextureData.PaintingTexture2D, GetTextureForMaterialOverride(TextureData));
 
 			// For the transactions add to the cache of overridden component.
 			PaintComponentsOverride.FindOrAdd(TextureData.PaintingTexture2D).PaintedComponents.Add(MeshComponent);
@@ -937,6 +946,16 @@ void UMeshTexturePaintingTool::StartPaintingTexture(UMeshComponent* InMeshCompon
 				const bool bForceLinearGamma = true;
 				TextureData->PaintRenderTargetTexture->InitCustomFormat(TextureWidth, TextureHeight, PF_A16B16G16R16, bForceLinearGamma);
 				TextureData->PaintRenderTargetTexture->UpdateResourceImmediate();
+
+				TextureData->PaintRenderTargetTextureAdapter = nullptr;
+				if (TextureData->PaintingTexture2D->IsCurrentlyVirtualTextured())
+				{
+					TextureData->PaintRenderTargetTextureAdapter = NewObject<UVirtualTextureAdapter>(GetTransientPackage(), NAME_None, RF_Transient);
+					TextureData->PaintRenderTargetTextureAdapter->Texture = TextureData->PaintRenderTargetTexture;
+					TextureData->PaintRenderTargetTextureAdapter->bUseCompressedFormat = true;
+					TextureData->PaintRenderTargetTextureAdapter->bUseDefaultTileSizes = true;
+					TextureData->PaintRenderTargetTextureAdapter->UpdateResource();
+				}
 			}
 			TextureData->PaintRenderTargetTexture->AddressX = Texture2D->AddressX;
 			TextureData->PaintRenderTargetTexture->AddressY = Texture2D->AddressY;
@@ -1392,6 +1411,12 @@ void UMeshTexturePaintingTool::PaintTexture(FMeshPaintParameters& InParams, TArr
 			TransitionAndCopyTexture(RHICmdList, RenderTargetResource->GetRenderTargetTexture(), RenderTargetResource->TextureRHI, {});
 		});
 	}
+
+	if (TextureData->PaintRenderTargetTextureAdapter)
+	{
+		TextureData->PaintRenderTargetTextureAdapter->Flush(FBox2f(FVector2f(0, 0), FVector2f(1, 1)));
+	}
+
 	FlushRenderingCommands();
 }
 
@@ -1484,7 +1509,7 @@ void UMeshTexturePaintingTool::OnTransactionStateChanged(const FTransactionConte
 								{
 									if (TSharedPtr<IMeshPaintComponentAdapter> PaintAdapterRessource = MeshPaintingSubsystem->GetAdapterForComponent(Component))
 									{
-										PaintAdapterRessource->ApplyOrRemoveTextureOverride(TextureData->PaintingTexture2D, TextureData->PaintRenderTargetTexture);
+										PaintAdapterRessource->ApplyOrRemoveTextureOverride(TextureData->PaintingTexture2D, GetTextureForMaterialOverride(*TextureData));
 										PaintComponentOverride.PaintedComponents.Add(Component);
 									}
 								}
@@ -1530,11 +1555,12 @@ void UMeshTexturePaintingTool::OnTransactionStateChanged(const FTransactionConte
 									MeshComponentsToRemove.Add(MeshComponent);
 								};
 
-								auto ComponentOverrideMissing = [this, MeshPaintingSubsystem, &ComponentOverride, Texture, &MeshComponentsToAdd, RenderTarget = TextureData->PaintRenderTargetTexture](TObjectPtr<UMeshComponent>& MeshComponent)
+								UTexture* OverrideTexture = GetTextureForMaterialOverride(*TextureData);
+								auto ComponentOverrideMissing = [this, MeshPaintingSubsystem, &ComponentOverride, Texture, &MeshComponentsToAdd, OverrideTexture](TObjectPtr<UMeshComponent>& MeshComponent)
 								{
 									if (TSharedPtr<IMeshPaintComponentAdapter> PaintAdapterRessource = MeshPaintingSubsystem->GetAdapterForComponent(MeshComponent))
 									{
-										PaintAdapterRessource->ApplyOrRemoveTextureOverride(Texture, RenderTarget);
+										PaintAdapterRessource->ApplyOrRemoveTextureOverride(Texture, OverrideTexture);
 									}
 
 									MeshComponentsToAdd.Add(MeshComponent);
diff --git a/Engine/Plugins/MeshPainting/Source/MeshPaintingToolset/Public/MeshPaintingToolsetTypes.h b/Engine/Plugins/MeshPainting/Source/MeshPaintingToolset/Public/MeshPaintingToolsetTypes.h
--- a/Engine/Plugins/MeshPainting/Source/MeshPaintingToolset/Public/MeshPaintingToolsetTypes.h
+++ b/Engine/Plugins/MeshPainting/Source/MeshPaintingToolset/Public/MeshPaintingToolsetTypes.h
@@ -153,6 +153,10 @@ struct FPaintTexture2DData
 	UPROPERTY(Transient)
 	TObjectPtr<UTextureRenderTarget2D> PaintBrushRenderTargetTexture = nullptr;
 
+	/** Optional virtual texture adapter that we can use to visualize the render target in materials that use mesh paint virtual textures */
+	UPROPERTY(Transient)
+	TObjectPtr<class UVirtualTextureAdapter> PaintRenderTargetTextureAdapter = nullptr;
+
 	FPaintTexture2DData() = default;
 
 	FPaintTexture2DData(UTexture2D* InPaintingTexture2D, bool InbIsPaintingTexture2DModified = false)
