MeshPaintTool always uses the tile sizes from the texture that it is painting to initialize the virtual texture adapter.
This is because we want the VT adapter to always use the same virtual texture pools as the texture it wraps.

diff --git a/Engine/Plugins/MeshPainting/Source/MeshPaintingToolset/Private/MeshTexturePaintingTool.cpp b/Engine/Plugins/MeshPainting/Source/MeshPaintingToolset/Private/MeshTexturePaintingTool.cpp
--- a/Engine/Plugins/MeshPainting/Source/MeshPaintingToolset/Private/MeshTexturePaintingTool.cpp
+++ b/Engine/Plugins/MeshPainting/Source/MeshPaintingToolset/Private/MeshTexturePaintingTool.cpp
@@ -22,6 +22,7 @@
 #include "Editor/TransBuffer.h"
 #include "Editor/UnrealEdEngine.h"
 #include "VT/VirtualTextureAdapter.h"
+#include "VT/VirtualTextureBuildSettings.h"
 #include "UnrealEdGlobals.h"
 
 #include "Editor/EditorEngine.h"
@@ -950,10 +951,15 @@ void UMeshTexturePaintingTool::StartPaintingTexture(UMeshComponent* InMeshCompon
 				TextureData->PaintRenderTargetTextureAdapter = nullptr;
 				if (TextureData->PaintingTexture2D->IsCurrentlyVirtualTextured())
 				{
+					FVirtualTextureBuildSettings VirtualTextureBuildSettings;
+					TextureData->PaintingTexture2D->GetVirtualTextureBuildSettings(VirtualTextureBuildSettings);
+
 					TextureData->PaintRenderTargetTextureAdapter = NewObject<UVirtualTextureAdapter>(GetTransientPackage(), NAME_None, RF_Transient);
 					TextureData->PaintRenderTargetTextureAdapter->Texture = TextureData->PaintRenderTargetTexture;
 					TextureData->PaintRenderTargetTextureAdapter->bUseCompressedFormat = true;
-					TextureData->PaintRenderTargetTextureAdapter->bUseDefaultTileSizes = true;
+					TextureData->PaintRenderTargetTextureAdapter->bUseDefaultTileSizes = false;
+					TextureData->PaintRenderTargetTextureAdapter->TileSize = VirtualTextureBuildSettings.TileSize;
+					TextureData->PaintRenderTargetTextureAdapter->TileBorderSize = VirtualTextureBuildSettings.TileBorderSize;
 					TextureData->PaintRenderTargetTextureAdapter->UpdateResource();
 				}
 			}
diff --git a/Engine/Source/Runtime/Engine/Private/VT/VirtualTextureAdapter.cpp b/Engine/Source/Runtime/Engine/Private/VT/VirtualTextureAdapter.cpp
--- a/Engine/Source/Runtime/Engine/Private/VT/VirtualTextureAdapter.cpp
+++ b/Engine/Source/Runtime/Engine/Private/VT/VirtualTextureAdapter.cpp
@@ -565,8 +565,8 @@ FTextureResource* UVirtualTextureAdapter::CreateResource()
 	FVirtualTextureBuildSettings DefaultSettings;
 	DefaultSettings.Init();
 
-	const uint32 FinalTileSize = bUseDefaultTileSizes ? DefaultSettings.TileSize : FMath::Min(FMath::RoundUpToPowerOfTwo(TileSize), 512u);
-	const uint32 FinalTileBorderSize = bUseDefaultTileSizes ? DefaultSettings.TileBorderSize : FMath::Min((((uint32)TileBorderSize + 1) & ~1u), 8u);
+	const uint32 FinalTileSize = bUseDefaultTileSizes ? DefaultSettings.TileSize : FVirtualTextureBuildSettings::ClampAndAlignTileSize(TileSize);
+	const uint32 FinalTileBorderSize = bUseDefaultTileSizes ? DefaultSettings.TileBorderSize : FVirtualTextureBuildSettings::ClampAndAlignTileBorderSize(TileBorderSize);
 	
 	return new FVirtualTextureAdapterRenderResource(this, Texture, FinalTileSize, FinalTileBorderSize, bUseCompressedFormat);
 }
