Enable selection of the mesh component mesh paint texture in the texture painting tool.
#rnx

diff --git a/Engine/Plugins/MeshPainting/Source/MeshPaintingToolset/Private/IMeshPaintComponentAdapter.cpp b/Engine/Plugins/MeshPainting/Source/MeshPaintingToolset/Private/IMeshPaintComponentAdapter.cpp
--- a/Engine/Plugins/MeshPainting/Source/MeshPaintingToolset/Private/IMeshPaintComponentAdapter.cpp
+++ b/Engine/Plugins/MeshPainting/Source/MeshPaintingToolset/Private/IMeshPaintComponentAdapter.cpp
@@ -8,6 +8,7 @@
 #include "Engine/World.h"
 #include "MaterialShared.h"
 #include "Materials/Material.h"
+#include "Materials/MaterialExpressionMeshPaintTextureObject.h"
 #include "Materials/MaterialExpressionTextureCoordinate.h"
 #include "Materials/MaterialExpressionTextureSampleParameter.h"
 #include "Materials/MaterialInstance.h"
@@ -105,6 +106,15 @@ void IMeshPaintComponentAdapter::DefaultQueryPaintableTextures(int32 MaterialInd
 						OutDefaultIndex = TextureIndex;
 					}
 				}
+
+				// If material samples the primitive mesh paint texture, then add it here.
+				if (UMaterialExpressionMeshPaintTextureObject* MeshPaintTextureExpression = Cast<UMaterialExpressionMeshPaintTextureObject>(Expression))
+				{
+					if (UTexture* MesPaintTexture = MeshComponent->GetMeshPaintTexture())
+					{
+						InOutTextureList.AddUnique(FPaintableTexture{ MesPaintTexture, 0 });
+					}
+				}
 			}
 		}
 		// Make sure to include all texture parameters, this will include all of the texture parameters from internal material functions
@@ -397,7 +407,7 @@ namespace UE::MeshPaintingToolset
 		return *this;
 	}
 
-	void FDefaultTextureOverride::ApplyOrRemoveTextureOverride(const UMeshComponent* InMeshComponent, const UTexture* SourceTexture, UTexture* OverrideTexture) const
+	void FDefaultTextureOverride::ApplyOrRemoveTextureOverride(UMeshComponent* InMeshComponent, const UTexture* SourceTexture, UTexture* OverrideTexture) const
 	{
 		check(IsInGameThread());
 
@@ -417,12 +427,17 @@ namespace UE::MeshPaintingToolset
 			{
 				// Keep track of the material overridden
 				Private::FGlobalTextureOverrideState::RegisterMaterialOverride(this, MaterialToCheck, SourceTexture, OverrideTexture, FeatureLevel);
-
 			}
 
 			++MaterialIndex;
 			MaterialToCheck = InMeshComponent->GetMaterial(MaterialIndex);
 		}
+
+		// Check to see if the source texture is the special mesh paint texture on the component.
+		if (InMeshComponent->GetMeshPaintTexture() == SourceTexture)
+		{
+			InMeshComponent->SetMeshPaintTextureOverride(OverrideTexture);
+		}
 	}
 
 	FDefaultTextureOverride::~FDefaultTextureOverride()
diff --git a/Engine/Plugins/MeshPainting/Source/MeshPaintingToolset/Public/IMeshPaintComponentAdapter.h b/Engine/Plugins/MeshPainting/Source/MeshPaintingToolset/Public/IMeshPaintComponentAdapter.h
--- a/Engine/Plugins/MeshPainting/Source/MeshPaintingToolset/Public/IMeshPaintComponentAdapter.h
+++ b/Engine/Plugins/MeshPainting/Source/MeshPaintingToolset/Public/IMeshPaintComponentAdapter.h
@@ -118,7 +118,7 @@ namespace UE::MeshPaintingToolset
 	class MESHPAINTINGTOOLSET_API FDefaultTextureOverride final
 	{
 	public:
-		void ApplyOrRemoveTextureOverride(const UMeshComponent* InMeshComponent, const UTexture* SourceTexture, UTexture* OverrideTexture) const;
+		void ApplyOrRemoveTextureOverride(UMeshComponent* InMeshComponent, const UTexture* SourceTexture, UTexture* OverrideTexture) const;
 
 		~FDefaultTextureOverride();
 
