Mesh Paint Tool support for texture color painting.
Added new 'Texture Color' tab in mesh paint tool which paints to the mesh component MeshPaintTexture.
This paint tool can paint to multiple selected components (similar to vertex paint tool).
Texture painting needed updating to support this multiselect behavior. In particular seam maps etc needed to be created per texture being edited rather than being global.
Renamed existing 'Color' and 'Weight' tools to 'Vertex Color' and 'Vertex Weight'. Existing 'Texture' tool is referred to as 'TextureAsset' in places where the distinction needs to be made clear in code. Similar renames were needed across the code.
Texture Color tool can add and remove MeshPaintTexture objects from components.
Texture Color tool can copy/paste MeshPaintTexture objects between components.
Texture Color tool can import vertex color painting to the MeshPaintTexture. This is done using the Geometry tools FMeshMapBaker.
Vertex Color tool can import texture color painting to the vertex colors on the component.
Generally tidied the code so that the mesh paint commands use a common pattern for their delegates. Starship icon names needed updating to match this.
Tidied up the MeshPaintingToolProperties so that they use a common hierachy between vertex and texture painting extensions.
Tidied up the MeshPaintingToolProperties UI customizations to take advantage of that sharing and remove duplicated code.

diff --git a/Engine/Plugins/MeshPainting/Source/MeshPaintEditorMode/MeshPaintEditorMode.Build.cs b/Engine/Plugins/MeshPainting/Source/MeshPaintEditorMode/MeshPaintEditorMode.Build.cs
--- a/Engine/Plugins/MeshPainting/Source/MeshPaintEditorMode/MeshPaintEditorMode.Build.cs
+++ b/Engine/Plugins/MeshPainting/Source/MeshPaintEditorMode/MeshPaintEditorMode.Build.cs
@@ -35,8 +35,12 @@ namespace UnrealBuildTool.Rules
                     "RHI",
 					"ToolWidgets",
 					"InterchangeEngine",
-					"InterchangePipelines"
-                }
+					"InterchangePipelines",
+					"DynamicMesh",
+					"GeometryCore",
+					"ImageCore",
+					"MeshConversionEngineTypes"
+				}
             );
 
         }
diff --git a/Engine/Plugins/MeshPainting/Source/MeshPaintEditorMode/Private/BrushSettingsCustomization.cpp b/Engine/Plugins/MeshPainting/Source/MeshPaintEditorMode/Private/BrushSettingsCustomization.cpp
--- a/Engine/Plugins/MeshPainting/Source/MeshPaintEditorMode/Private/BrushSettingsCustomization.cpp
+++ b/Engine/Plugins/MeshPainting/Source/MeshPaintEditorMode/Private/BrushSettingsCustomization.cpp
@@ -38,44 +38,34 @@ TSharedRef<SHorizontalBox> CreateColorChannelWidget(TSharedRef<IPropertyHandle>
 }
 
 
-TSharedRef<IDetailCustomization> FVertexPaintingSettingsCustomization::MakeInstance()
+TSharedRef<IDetailCustomization> FMeshPaintingSettingsCustomization::MakeInstance()
 {
-	return MakeShareable(new FVertexPaintingSettingsCustomization);
+	return MakeShareable(new FMeshPaintingSettingsCustomization);
 }
 
-void FVertexPaintingSettingsCustomization::CustomizeDetails(IDetailLayoutBuilder& DetailLayout)
-{	
-
+void FMeshPaintingSettingsCustomization::CustomizeDetails(IDetailLayoutBuilder& DetailLayout)
+{
 	IDetailCategoryBuilder& BrushCategory = DetailLayout.EditCategory(TEXT("Brush"), FText::GetEmpty(), ECategoryPriority::Important);
-	IDetailCategoryBuilder& VertexCategory = DetailLayout.EditCategory(TEXT("VertexPainting"));
 
-	VertexCategory.AddCustomRow(NSLOCTEXT("VertexPaintSettings", "InstanceColorSize", "Instance Color Size"))
-		.WholeRowContent()
-		[
-			SNew(STextBlock)
-			.Text_Lambda([]() -> FText { return FText::Format(FTextFormat::FromString(TEXT("Instance Color Size: {0} KB")), UMeshPaintMode::GetMeshPaintMode()->GetCachedVertexDataSize() / 1024.f); })
-		];
-
-	/** Creates a custom widget row containing all color channel flags */	
-	TSharedRef<IPropertyHandle> PaintColor = DetailLayout.GetProperty("PaintColor", UMeshVertexPaintingToolProperties::StaticClass());
+	// Customize paint color with a swap button
+	TSharedRef<IPropertyHandle> PaintColor = DetailLayout.GetProperty("PaintColor", UMeshPaintingToolProperties::StaticClass());
 	PaintColor->MarkHiddenByCustomization();
-	TSharedRef<IPropertyHandle> EraseColor = DetailLayout.GetProperty("EraseColor", UMeshVertexPaintingToolProperties::StaticClass());
+	TSharedRef<IPropertyHandle> EraseColor = DetailLayout.GetProperty("EraseColor", UMeshPaintingToolProperties::StaticClass());
 	EraseColor->MarkHiddenByCustomization();
 
-	// Customize paint color with a swap button
 	{
 		TSharedPtr<SWidget> NameWidget;
 		TSharedPtr<SWidget> ValueWidget;
-		
-		IDetailPropertyRow& PaintColorProp = VertexCategory.AddProperty(PaintColor);
+
+		IDetailPropertyRow& PaintColorProp = BrushCategory.AddProperty(PaintColor);
 		PaintColorProp.GetDefaultWidgets(NameWidget, ValueWidget, false);
-		FDetailWidgetRow& Row = PaintColorProp.CustomWidget(true);
-		Row.NameContent()
+		FDetailWidgetRow& PaintRow = PaintColorProp.CustomWidget(true);
+		PaintRow.NameContent()
 		[
 			NameWidget.ToSharedRef()
 		];
 
-		Row.ValueContent()
+		PaintRow.ValueContent()
 		.MinDesiredWidth(250)
 		.MaxDesiredWidth(0)
 		[
@@ -99,7 +89,7 @@ void FVertexPaintingSettingsCustomization::CustomizeDetails(IDetailLayoutBuilder
 				.ToolTipText(NSLOCTEXT("VertexPaintSettings", "SwapColors", "Swap Paint and Erase Colors"))
 				.HAlign(HAlign_Center)
 				.VAlign(VAlign_Center)
-				.OnClicked(this, &FVertexPaintingSettingsCustomization::OnSwapColorsClicked, PaintColor, EraseColor)
+				.OnClicked(this, &FMeshPaintingSettingsCustomization::OnSwapColorsClicked, PaintColor, EraseColor)
 				.ContentPadding(0)
 				[
 					SNew(SImage).Image(FAppStyle::GetBrush("MeshPaint.Swap"))
@@ -107,52 +97,23 @@ void FVertexPaintingSettingsCustomization::CustomizeDetails(IDetailLayoutBuilder
 			]
 		];
 	}
-
 	{
-		IDetailPropertyRow& EraseColorProp = VertexCategory.AddProperty(EraseColor);
+		IDetailPropertyRow& EraseColorProp = BrushCategory.AddProperty(EraseColor);
 
 		TSharedPtr<SWidget> NameWidget;
 		TSharedPtr<SWidget> ValueWidget;
-		
+
 		FDetailWidgetRow& Row = EraseColorProp.CustomWidget(true);
 		Row.ValueContent().MinDesiredWidth(250 - 16.f);
 		EraseColorProp.GetDefaultWidgets(NameWidget, ValueWidget, Row, false);
 	}
-
-}
-
-void FVertexPaintingSettingsCustomization::OnTextureWeightTypeChanged(TSharedRef<IPropertyHandle> WeightTypeProperty, TSharedRef<IPropertyHandle> PaintWeightProperty, TSharedRef<IPropertyHandle> EraseWeightProperty)
-{
-	UEnum* ImportTypeEnum = StaticEnum<ETexturePaintWeightIndex>();
-	uint8 EnumValue = 0;
-	WeightTypeProperty->GetValue(EnumValue);	
-
-	BlendPaintEnumRestriction->RemoveAll();
-	for (uint8 EnumIndex = 0; EnumIndex < (ImportTypeEnum->GetMaxEnumValue() + 1); ++EnumIndex)
-	{
-		if ((EnumIndex + 1) > EnumValue)
-		{
-			FString EnumName = ImportTypeEnum->GetNameByValue(EnumIndex).ToString();
-			EnumName.RemoveFromStart("ETexturePaintIndex::");
-			BlendPaintEnumRestriction->AddDisabledValue(EnumName);
-		}
-	}
-
-	uint8 Value = 0;
-	PaintWeightProperty->GetValue(Value);
-	Value = FMath::Clamp<uint8>(Value, 0, EnumValue - 1);
-	PaintWeightProperty->SetValue(Value);
-
-	EraseWeightProperty->GetValue(Value);
-	Value = FMath::Clamp<uint8>(Value, 0, EnumValue - 1);
-	EraseWeightProperty->SetValue(Value);
 }
 
-FReply FVertexPaintingSettingsCustomization::OnSwapColorsClicked(TSharedRef<IPropertyHandle> PaintColor, TSharedRef<IPropertyHandle> EraseColor)
+FReply FMeshPaintingSettingsCustomization::OnSwapColorsClicked(TSharedRef<IPropertyHandle> PaintColor, TSharedRef<IPropertyHandle> EraseColor)
 {
-	FScopedTransaction Transaction(NSLOCTEXT("VertexPaintSettings", "SwapColorsTransation", "Swap paint and erase colors"));
+	FScopedTransaction Transaction(NSLOCTEXT("MeshPaintSettings", "SwapColorsTransation", "Swap paint and erase colors"));
 
-	GEditor->GetEditorSubsystem<UMeshPaintModeSubsystem>()->SwapVertexColors();
+	GEditor->GetEditorSubsystem<UMeshPaintModeSubsystem>()->SwapColors();
 	UMeshVertexPaintingToolProperties* Settings = UMeshPaintMode::GetVertexToolProperties();
 	if (Settings)
 	{
@@ -163,26 +124,45 @@ FReply FVertexPaintingSettingsCustomization::OnSwapColorsClicked(TSharedRef<IPro
 }
 
 
+TSharedRef<IDetailCustomization> FVertexPaintingSettingsCustomization::MakeInstance()
+{
+	return MakeShareable(new FVertexPaintingSettingsCustomization);
+}
 
-TSharedRef<IDetailCustomization> FColorPaintingSettingsCustomization::MakeInstance()
+void FVertexPaintingSettingsCustomization::CustomizeDetails(IDetailLayoutBuilder& DetailLayout)
 {
-	return MakeShareable(new FColorPaintingSettingsCustomization);
+	FMeshPaintingSettingsCustomization::CustomizeDetails(DetailLayout);
+
+	IDetailCategoryBuilder& VertexCategory = DetailLayout.EditCategory(TEXT("VertexPainting"));
+
+	VertexCategory.AddCustomRow(NSLOCTEXT("VertexPaintSettings", "InstanceColorSize", "Instance Color Size"))
+		.WholeRowContent()
+		[
+			SNew(STextBlock)
+			.Text_Lambda([]() -> FText { return FText::Format(FTextFormat::FromString(TEXT("Instance Color Size: {0} KB")), UMeshPaintMode::GetMeshPaintMode()->GetCachedVertexDataSize() / 1024.f); })
+		];
+}
+
+
+TSharedRef<IDetailCustomization> FVertexColorPaintingSettingsCustomization::MakeInstance()
+{
+	return MakeShareable(new FVertexColorPaintingSettingsCustomization);
 }
 
-void FColorPaintingSettingsCustomization::CustomizeDetails(IDetailLayoutBuilder& DetailLayout)
+void FVertexColorPaintingSettingsCustomization::CustomizeDetails(IDetailLayoutBuilder& DetailLayout)
 {
 	FVertexPaintingSettingsCustomization::CustomizeDetails(DetailLayout);
 
 	IDetailCategoryBuilder& ColorCategory = DetailLayout.EditCategory(TEXT("ColorPainting"));
 
 	/** Creates a custom widget row containing all color channel flags */
-	TSharedRef<IPropertyHandle> RedChannel = DetailLayout.GetProperty("bWriteRed", UMeshColorPaintingToolProperties::StaticClass());
+	TSharedRef<IPropertyHandle> RedChannel = DetailLayout.GetProperty("bWriteRed", UMeshVertexColorPaintingToolProperties::StaticClass());
 	RedChannel->MarkHiddenByCustomization();
-	TSharedRef<IPropertyHandle> GreenChannel = DetailLayout.GetProperty("bWriteGreen", UMeshColorPaintingToolProperties::StaticClass());
+	TSharedRef<IPropertyHandle> GreenChannel = DetailLayout.GetProperty("bWriteGreen", UMeshVertexColorPaintingToolProperties::StaticClass());
 	GreenChannel->MarkHiddenByCustomization();
-	TSharedRef<IPropertyHandle> BlueChannel = DetailLayout.GetProperty("bWriteBlue", UMeshColorPaintingToolProperties::StaticClass());
+	TSharedRef<IPropertyHandle> BlueChannel = DetailLayout.GetProperty("bWriteBlue", UMeshVertexColorPaintingToolProperties::StaticClass());
 	BlueChannel->MarkHiddenByCustomization();
-	TSharedRef<IPropertyHandle> AlphaChannel = DetailLayout.GetProperty("bWriteAlpha", UMeshColorPaintingToolProperties::StaticClass());
+	TSharedRef<IPropertyHandle> AlphaChannel = DetailLayout.GetProperty("bWriteAlpha", UMeshVertexColorPaintingToolProperties::StaticClass());
 	AlphaChannel->MarkHiddenByCustomization();
 	TArray<TSharedRef<IPropertyHandle>> Channels = { RedChannel, GreenChannel, BlueChannel, AlphaChannel };
 	TSharedPtr<SHorizontalBox> ChannelsWidget;
@@ -203,20 +183,20 @@ void FColorPaintingSettingsCustomization::CustomizeDetails(IDetailLayoutBuilder&
 	for (TSharedRef<IPropertyHandle> Channel : Channels)
 	{
 		ChannelsWidget->AddSlot()
-			.AutoWidth()
-			.Padding(0.0f, 0.0f, 4.0f, 0.0f)
-			[
-				CreateColorChannelWidget(Channel)
-			];
+		.AutoWidth()
+		.Padding(0.0f, 0.0f, 4.0f, 0.0f)
+		[
+			CreateColorChannelWidget(Channel)
+		];
 	}
 
 	static FText RestrictReason = NSLOCTEXT("VertexPaintSettings", "TextureIndexRestriction", "Unable to paint this Texture, change Texture Weight Type");
 	BlendPaintEnumRestriction = MakeShareable(new FPropertyRestriction(RestrictReason));
 
 	/** Add custom row for painting on specific LOD level with callbacks to the painter to update the data */
-	TSharedRef<IPropertyHandle> LODPaintingEnabled = DetailLayout.GetProperty("bPaintOnSpecificLOD", UMeshColorPaintingToolProperties::StaticClass());
+	TSharedRef<IPropertyHandle> LODPaintingEnabled = DetailLayout.GetProperty("bPaintOnSpecificLOD", UMeshVertexColorPaintingToolProperties::StaticClass());
 	LODPaintingEnabled->MarkHiddenByCustomization();
-	TSharedRef<IPropertyHandle> LODPaintingIndex = DetailLayout.GetProperty("LODIndex", UMeshColorPaintingToolProperties::StaticClass());
+	TSharedRef<IPropertyHandle> LODPaintingIndex = DetailLayout.GetProperty("LODIndex", UMeshVertexColorPaintingToolProperties::StaticClass());
 	LODPaintingIndex->MarkHiddenByCustomization();
 	TSharedPtr<SWidget> LODIndexWidget = LODPaintingIndex->CreatePropertyValueWidget();
 
@@ -227,7 +207,7 @@ void FColorPaintingSettingsCustomization::CustomizeDetails(IDetailLayoutBuilder&
 			.Text(NSLOCTEXT("LODPainting", "LODPaintingSetupLabel", "LOD Model Painting"))
 			.ToolTipText(NSLOCTEXT("LODPainting", "LODPaintingSetupToolTip", "Allows for Painting Vertex Colors on Specific LOD Models."))
 		]
-	.ValueContent()
+		.ValueContent()
 		[
 			SNew(SHorizontalBox)
 			+ SHorizontalBox::Slot()
@@ -235,7 +215,7 @@ void FColorPaintingSettingsCustomization::CustomizeDetails(IDetailLayoutBuilder&
 		.AutoWidth()
 		[
 			SNew(SCheckBox)
-				.IsChecked_Lambda([=]() -> ECheckBoxState { return (UMeshPaintMode::GetColorToolProperties() && UMeshPaintMode::GetColorToolProperties()->bPaintOnSpecificLOD ? ECheckBoxState::Checked : ECheckBoxState::Unchecked); })
+				.IsChecked_Lambda([=]() -> ECheckBoxState { return (UMeshPaintMode::GetVertexColorToolProperties() && UMeshPaintMode::GetVertexColorToolProperties()->bPaintOnSpecificLOD ? ECheckBoxState::Checked : ECheckBoxState::Unchecked); })
 				.OnCheckStateChanged(FOnCheckStateChanged::CreateLambda([=](ECheckBoxState State) { 
 					if (UMeshColorPaintingTool* ColorBrush = Cast<UMeshColorPaintingTool>(UMeshPaintMode::GetMeshPaintMode()->GetToolManager()->GetActiveTool(EToolSide::Left)))
 					{
@@ -243,13 +223,13 @@ void FColorPaintingSettingsCustomization::CustomizeDetails(IDetailLayoutBuilder&
 					}
 				}))
 		]
-	+ SHorizontalBox::Slot()
+		+ SHorizontalBox::Slot()
 		.Padding(0.0f, 0.0f, 4.0f, 0.0f)
 		[
 			SNew(SNumericEntryBox<int32>)
-			.IsEnabled_Lambda([=]() -> bool { return UMeshPaintMode::GetColorToolProperties() ? UMeshPaintMode::GetColorToolProperties()->bPaintOnSpecificLOD : false;  })
+			.IsEnabled_Lambda([=]() -> bool { return UMeshPaintMode::GetVertexColorToolProperties() ? UMeshPaintMode::GetVertexColorToolProperties()->bPaintOnSpecificLOD : false;  })
 			.AllowSpin(true)
-			.Value_Lambda([=]() -> int32 { return UMeshPaintMode::GetColorToolProperties() ? UMeshPaintMode::GetColorToolProperties()->LODIndex : 0; })
+			.Value_Lambda([=]() -> int32 { return UMeshPaintMode::GetVertexColorToolProperties() ? UMeshPaintMode::GetVertexColorToolProperties()->LODIndex : 0; })
 			.MinValue(0)
 			.MaxValue_Lambda([=]() -> int32 { 
 					if (UMeshColorPaintingTool* ColorBrush = Cast<UMeshColorPaintingTool>(UMeshPaintMode::GetMeshPaintMode()->GetToolManager()->GetActiveTool(EToolSide::Left)))
@@ -265,9 +245,9 @@ void FColorPaintingSettingsCustomization::CustomizeDetails(IDetailLayoutBuilder&
 					}
 					return INT_MAX;
 				})
-			.OnValueChanged(SNumericEntryBox<int32>::FOnValueChanged::CreateLambda([=](int32 Value) { UMeshPaintMode::GetColorToolProperties()->LODIndex = Value; }))
+			.OnValueChanged(SNumericEntryBox<int32>::FOnValueChanged::CreateLambda([=](int32 Value) { UMeshPaintMode::GetVertexColorToolProperties()->LODIndex = Value; }))
 			.OnValueCommitted(SNumericEntryBox<int32>::FOnValueCommitted::CreateLambda([=](int32 Value, ETextCommit::Type CommitType) { 
-					UMeshPaintMode::GetColorToolProperties()->LODIndex = Value; 
+					UMeshPaintMode::GetVertexColorToolProperties()->LODIndex = Value; 
 					if (UMeshColorPaintingTool* ColorBrush = Cast<UMeshColorPaintingTool>(UMeshPaintMode::GetMeshPaintMode()->GetToolManager()->GetActiveTool(EToolSide::Left)))
 					{
 						ColorBrush->PaintLODChanged();
@@ -282,7 +262,7 @@ void FColorPaintingSettingsCustomization::CustomizeDetails(IDetailLayoutBuilder&
 			SNew(SWarningOrErrorBox)
 			.Visibility_Lambda([this]() -> EVisibility
 			{
-				if (UMeshColorPaintingToolProperties* ColorProperties = UMeshPaintMode::GetColorToolProperties())
+				if (UMeshVertexColorPaintingToolProperties* ColorProperties = UMeshPaintMode::GetVertexColorToolProperties())
 				{
 					return ColorProperties->bPaintOnSpecificLOD ? EVisibility::Collapsed : EVisibility::Visible;
 				}
@@ -296,7 +276,7 @@ void FColorPaintingSettingsCustomization::CustomizeDetails(IDetailLayoutBuilder&
 
 				const bool bGeometryCollectionText = UMeshPaintMode::GetMeshPaintMode()->GetSelectedComponents<UGeometryCollectionComponent>().Num() > 0;
 				const bool bSkelMeshText = UMeshPaintMode::GetMeshPaintMode()->GetSelectedComponents<USkeletalMeshComponent>().Num() > 0;
-				const bool bLODPaintText = UMeshPaintMode::GetColorToolProperties() ?  !UMeshPaintMode::GetColorToolProperties()->bPaintOnSpecificLOD : false;
+				const bool bLODPaintText = UMeshPaintMode::GetVertexColorToolProperties() ?  !UMeshPaintMode::GetVertexColorToolProperties()->bPaintOnSpecificLOD : false;
 				return FText::Format(FTextFormat::FromString(TEXT("{0}{1}{2}{3}")), 
 					bSkelMeshText ? SkelMeshNotificationText : FText::GetEmpty(),
 					bGeometryCollectionText ? GeometryCollectionNotificationText : FText::GetEmpty(),
@@ -307,20 +287,23 @@ void FColorPaintingSettingsCustomization::CustomizeDetails(IDetailLayoutBuilder&
 		];
 }
 
-TSharedRef<IDetailCustomization> FWeightPaintingSettingsCustomization::MakeInstance()
+
+TSharedRef<IDetailCustomization> FVertexWeightPaintingSettingsCustomization::MakeInstance()
 {
-	return MakeShareable(new FWeightPaintingSettingsCustomization);
+	return MakeShareable(new FVertexWeightPaintingSettingsCustomization);
 }
 
-void FWeightPaintingSettingsCustomization::CustomizeDetails(IDetailLayoutBuilder& DetailLayout)
+void FVertexWeightPaintingSettingsCustomization::CustomizeDetails(IDetailLayoutBuilder& DetailLayout)
 {
 	FVertexPaintingSettingsCustomization::CustomizeDetails(DetailLayout);
 
-	TSharedRef<IPropertyHandle> WeightTypeProperty = DetailLayout.GetProperty(GET_MEMBER_NAME_CHECKED(UMeshWeightPaintingToolProperties, TextureWeightType));
-	TSharedRef<IPropertyHandle> PaintWeightProperty = DetailLayout.GetProperty(GET_MEMBER_NAME_CHECKED(UMeshWeightPaintingToolProperties, PaintTextureWeightIndex));
-	TSharedRef<IPropertyHandle> EraseWeightProperty = DetailLayout.GetProperty(GET_MEMBER_NAME_CHECKED(UMeshWeightPaintingToolProperties, EraseTextureWeightIndex));
+	IDetailCategoryBuilder& WeightCategory = DetailLayout.EditCategory(TEXT("WeightPainting"));
+
+	TSharedRef<IPropertyHandle> WeightTypeProperty = DetailLayout.GetProperty(GET_MEMBER_NAME_CHECKED(UMeshVertexWeightPaintingToolProperties, TextureWeightType));
+	TSharedRef<IPropertyHandle> PaintWeightProperty = DetailLayout.GetProperty(GET_MEMBER_NAME_CHECKED(UMeshVertexWeightPaintingToolProperties, PaintTextureWeightIndex));
+	TSharedRef<IPropertyHandle> EraseWeightProperty = DetailLayout.GetProperty(GET_MEMBER_NAME_CHECKED(UMeshVertexWeightPaintingToolProperties, EraseTextureWeightIndex));
 
-	WeightTypeProperty->SetOnPropertyValueChanged(FSimpleDelegate::CreateRaw(this, &FWeightPaintingSettingsCustomization::OnTextureWeightTypeChanged, WeightTypeProperty, PaintWeightProperty, EraseWeightProperty));
+	WeightTypeProperty->SetOnPropertyValueChanged(FSimpleDelegate::CreateRaw(this, &FVertexWeightPaintingSettingsCustomization::OnTextureWeightTypeChanged, WeightTypeProperty, PaintWeightProperty, EraseWeightProperty));
 
 	static FText RestrictReason = NSLOCTEXT("VertexPaintSettings", "TextureIndexRestriction", "Unable to paint this Texture, change Texture Weight Type");
 	BlendPaintEnumRestriction = MakeShareable(new FPropertyRestriction(RestrictReason));
@@ -330,7 +313,7 @@ void FWeightPaintingSettingsCustomization::CustomizeDetails(IDetailLayoutBuilder
 	OnTextureWeightTypeChanged(WeightTypeProperty, PaintWeightProperty, EraseWeightProperty);
 }
 
-void FWeightPaintingSettingsCustomization::OnTextureWeightTypeChanged(TSharedRef<IPropertyHandle> WeightTypeProperty, TSharedRef<IPropertyHandle> PaintWeightProperty, TSharedRef<IPropertyHandle> EraseWeightProperty)
+void FVertexWeightPaintingSettingsCustomization::OnTextureWeightTypeChanged(TSharedRef<IPropertyHandle> WeightTypeProperty, TSharedRef<IPropertyHandle> PaintWeightProperty, TSharedRef<IPropertyHandle> EraseWeightProperty)
 {
 	UEnum* ImportTypeEnum = StaticEnum<EMeshPaintTextureIndex>();
 	uint8 EnumValue = 0;
@@ -365,8 +348,11 @@ TSharedRef<IDetailCustomization> FTexturePaintingSettingsCustomization::MakeInst
 
 void FTexturePaintingSettingsCustomization::CustomizeDetails(IDetailLayoutBuilder& DetailLayout)
 {
+	FMeshPaintingSettingsCustomization::CustomizeDetails(DetailLayout);
+
 	IDetailCategoryBuilder& BrushCategory = DetailLayout.EditCategory(TEXT("Brush"), FText::GetEmpty(), ECategoryPriority::Important);
 	IDetailCategoryBuilder& TextureCategory = DetailLayout.EditCategory(TEXT("TexturePainting"));
+	IDetailCategoryBuilder& ColorCategory = DetailLayout.EditCategory(TEXT("ColorPainting"));
 
 	TSharedRef<IPropertyHandle> RedChannel = DetailLayout.GetProperty("bWriteRed", UMeshTexturePaintingToolProperties::StaticClass());
 	RedChannel->MarkHiddenByCustomization();
@@ -379,176 +365,118 @@ void FTexturePaintingSettingsCustomization::CustomizeDetails(IDetailLayoutBuilde
 	TArray<TSharedRef<IPropertyHandle>> Channels = { RedChannel, GreenChannel, BlueChannel, AlphaChannel };
 	TSharedPtr<SHorizontalBox> ChannelsWidget;
 
-	TextureCategory.AddCustomRow(NSLOCTEXT("VertexPaintSettings", "ChannelLabel", "Channels"))
-		.NameContent()
-		[
-			SNew(STextBlock)
-			.Text(NSLOCTEXT("VertexPaintSettings", "ChannelsLabel", "Channels"))
-			.ToolTipText(NSLOCTEXT("VertexPaintSettings", "ChannelsToolTip", "Colors Channels which should be influenced during Painting."))
-			.Font(IDetailLayoutBuilder::GetDetailFont())
-		]
-		.ValueContent()
-		.MaxDesiredWidth(250.0f)
-		[
-			SAssignNew(ChannelsWidget, SHorizontalBox)
-		];
+	ColorCategory.AddCustomRow(NSLOCTEXT("VertexPaintSettings", "ChannelLabel", "Channels"))
+	.NameContent()
+	[
+		SNew(STextBlock)
+		.Text(NSLOCTEXT("VertexPaintSettings", "ChannelsLabel", "Channels"))
+		.ToolTipText(NSLOCTEXT("VertexPaintSettings", "ChannelsToolTip", "Colors Channels which should be influenced during Painting."))
+		.Font(IDetailLayoutBuilder::GetDetailFont())
+	]
+	.ValueContent()
+	.MaxDesiredWidth(250.0f)
+	[
+		SAssignNew(ChannelsWidget, SHorizontalBox)
+	];
 
 	for (TSharedRef<IPropertyHandle> Channel : Channels)
 	{
 		ChannelsWidget->AddSlot()
-			.AutoWidth()
-			.Padding(0.0f, 0.0f, 4.0f, 0.0f)
-			[
-				CreateColorChannelWidget(Channel)
-			];
+		.AutoWidth()
+		.Padding(0.0f, 0.0f, 4.0f, 0.0f)
+		[
+			CreateColorChannelWidget(Channel)
+		];
 	}
+}
+
+
+TSharedRef<IDetailCustomization> FTextureColorPaintingSettingsCustomization::MakeInstance()
+{
+	return MakeShareable(new FTextureColorPaintingSettingsCustomization);
+}
 
+void FTextureColorPaintingSettingsCustomization::CustomizeDetails(IDetailLayoutBuilder& DetailLayout)
+{
+	FTexturePaintingSettingsCustomization::CustomizeDetails(DetailLayout);
+}
 
-	TSharedRef<IPropertyHandle> UVChannel = DetailLayout.GetProperty(GET_MEMBER_NAME_CHECKED(UMeshTexturePaintingToolProperties, UVChannel));
+
+TSharedRef<IDetailCustomization> FTextureAssetPaintingSettingsCustomization::MakeInstance()
+{
+	return MakeShareable(new FTextureAssetPaintingSettingsCustomization);
+}
+
+void FTextureAssetPaintingSettingsCustomization::CustomizeDetails(IDetailLayoutBuilder& DetailLayout)
+{
+	FTexturePaintingSettingsCustomization::CustomizeDetails(DetailLayout);
+
+	IDetailCategoryBuilder& TextureCategory = DetailLayout.EditCategory(TEXT("TexturePainting"));
+
+	TSharedRef<IPropertyHandle> UVChannel = DetailLayout.GetProperty(GET_MEMBER_NAME_CHECKED(UMeshTextureAssetPaintingToolProperties, UVChannel));
 	UVChannel->MarkHiddenByCustomization();
 
 	TextureCategory.AddCustomRow(LOCTEXT("TexturePaintingUVLabel", "Texture Painting UV Channel"))
-		.NameContent()
-		[
-			UVChannel->CreatePropertyNameWidget()
-		]
+	.NameContent()
+	[
+		UVChannel->CreatePropertyNameWidget()
+	]
 	.ValueContent()
-		[
-			SNew(SHorizontalBox)
-			+ SHorizontalBox::Slot()
+	[
+		SNew(SHorizontalBox)
+		+ SHorizontalBox::Slot()
 		.Padding(0.0f, 0.0f, 4.0f, 0.0f)
 		[
 			SNew(SNumericEntryBox<int32>)
 			.Font(IDetailLayoutBuilder::GetDetailFont())
 			.AllowSpin(true)
-			.Value_Lambda([]() -> int32 { return UMeshPaintMode::GetTextureToolProperties() ? UMeshPaintMode::GetTextureToolProperties()->UVChannel : 0; })
+			.Value_Lambda([]() -> int32 { return UMeshPaintMode::GetTextureAssetToolProperties() ? UMeshPaintMode::GetTextureAssetToolProperties()->UVChannel : 0; })
 			.MinValue(0)
 			.MaxValue_Lambda([]() -> int32 { return GEngine->GetEngineSubsystem<UMeshPaintingSubsystem>()->GetMaxUVIndexToPaint(); })
-			.OnValueChanged(SNumericEntryBox<int32>::FOnValueChanged::CreateLambda([=](int32 Value) { 
-				if (UMeshPaintMode::GetTextureToolProperties())
-				{
-					UMeshPaintMode::GetTextureToolProperties()->UVChannel = Value;
-				}
+			.OnValueChanged(SNumericEntryBox<int32>::FOnValueChanged::CreateLambda([=](int32 Value) {
+			if (UMeshPaintMode::GetTextureAssetToolProperties())
+			{
+				UMeshPaintMode::GetTextureAssetToolProperties()->UVChannel = Value;
+			}
 				}))
 			.OnValueCommitted(SNumericEntryBox<int32>::FOnValueCommitted::CreateLambda([](int32 Value, ETextCommit::Type CommitType) {
-					if (UMeshPaintMode::GetTextureToolProperties())
-					{
-						UMeshPaintMode::GetTextureToolProperties()->UVChannel = Value;
-					}
-				}))
-			]
-		];
-
+			if (UMeshPaintMode::GetTextureAssetToolProperties())
+			{
+				UMeshPaintMode::GetTextureAssetToolProperties()->UVChannel = Value;
+			}
+		}))
+		]
+	];
 
-	TSharedRef<IPropertyHandle> TextureProperty = DetailLayout.GetProperty(GET_MEMBER_NAME_CHECKED(UMeshTexturePaintingToolProperties, PaintTexture));
+	TSharedRef<IPropertyHandle> TextureProperty = DetailLayout.GetProperty(GET_MEMBER_NAME_CHECKED(UMeshTextureAssetPaintingToolProperties, PaintTexture));
 	TextureProperty->MarkHiddenByCustomization();
+
 	TSharedPtr<SHorizontalBox> TextureWidget;
 	FDetailWidgetRow& Row = TextureCategory.AddCustomRow(NSLOCTEXT("TexturePaintSetting", "TextureSearchString", "Texture"))
-		.NameContent()
-		[
-			SNew(STextBlock)
-			.Text(NSLOCTEXT("TexturePaintSettings", "PaintTextureLabel", "Paint Texture"))
-			.ToolTipText(NSLOCTEXT("TexturePaintSettings", "PaintTextureToolTip", "Texture to Apply Painting to."))
-			.Font(IDetailLayoutBuilder::GetDetailFont())
-		]
-		.ValueContent()
-		.MaxDesiredWidth(250.0f)
-		[
-			SAssignNew(TextureWidget, SHorizontalBox)
-		];
+	.NameContent()
+	[
+		SNew(STextBlock)
+		.Text(NSLOCTEXT("TexturePaintSettings", "PaintTextureLabel", "Paint Texture"))
+		.ToolTipText(NSLOCTEXT("TexturePaintSettings", "PaintTextureToolTip", "Texture to Apply Painting to."))
+		.Font(IDetailLayoutBuilder::GetDetailFont())
+	]
+	.ValueContent()
+	.MaxDesiredWidth(250.0f)
+	[
+		SAssignNew(TextureWidget, SHorizontalBox)
+	];
 
 	/** Use a SObjectPropertyEntryBox to benefit from its functionality */
-		TextureWidget->AddSlot()
-		[
-			SNew(SObjectPropertyEntryBox)
-			.PropertyHandle(TextureProperty)
-			.AllowedClass(UTexture2D::StaticClass())
-			.OnShouldFilterAsset(FOnShouldFilterAsset::CreateUObject(Cast<UMeshTexturePaintingTool>(UMeshPaintMode::GetMeshPaintMode()->GetToolManager()->GetActiveTool(EToolSide::Left)), &UMeshTexturePaintingTool::ShouldFilterTextureAsset))
-			.OnObjectChanged(FOnSetObject::CreateUObject(Cast<UMeshTexturePaintingTool>(UMeshPaintMode::GetMeshPaintMode()->GetToolManager()->GetActiveTool(EToolSide::Left)), &UMeshTexturePaintingTool::PaintTextureChanged))
-			.DisplayUseSelected(false)
-			.ThumbnailPool(DetailLayout.GetThumbnailPool())
-		];
-
-		/** Creates a custom widget row containing all color channel flags */
-		TSharedRef<IPropertyHandle> PaintColor = DetailLayout.GetProperty("PaintColor", UMeshTexturePaintingToolProperties::StaticClass());
-		PaintColor->MarkHiddenByCustomization();
-		TSharedRef<IPropertyHandle> EraseColor = DetailLayout.GetProperty("EraseColor", UMeshTexturePaintingToolProperties::StaticClass());
-		EraseColor->MarkHiddenByCustomization();
-
-		// Customize paint color with a swap button
-		{
-			TSharedPtr<SWidget> NameWidget;
-			TSharedPtr<SWidget> ValueWidget;
-
-			IDetailPropertyRow& PaintColorProp = TextureCategory.AddProperty(PaintColor);
-			PaintColorProp.GetDefaultWidgets(NameWidget, ValueWidget, false);
-			FDetailWidgetRow& PaintRow = PaintColorProp.CustomWidget(true);
-			PaintRow.NameContent()
-				[
-					NameWidget.ToSharedRef()
-				];
-
-			PaintRow.ValueContent()
-				.MinDesiredWidth(250)
-				.MaxDesiredWidth(0)
-				[
-					SNew(SHorizontalBox)
-					+ SHorizontalBox::Slot()
-					.FillWidth(1.0)
-					.HAlign(HAlign_Left)
-					[
-						SNew(SBox)
-						.WidthOverride(250.f)
-						[
-							ValueWidget.ToSharedRef()
-						]
-					]
-					+ SHorizontalBox::Slot()
-					.HAlign(HAlign_Center)
-					.AutoWidth()
-					[
-						SNew(SButton)
-						.ButtonStyle(FAppStyle::Get(), "HoverHintOnly")
-						.ToolTipText(NSLOCTEXT("VertexPaintSettings", "SwapColors", "Swap Paint and Erase Colors"))
-						.HAlign(HAlign_Center)
-						.VAlign(VAlign_Center)
-						.OnClicked(this, &FTexturePaintingSettingsCustomization::OnSwapColorsClicked, PaintColor, EraseColor)
-						.ContentPadding(0)
-						[
-							SNew(SImage).Image(FAppStyle::GetBrush("MeshPaint.Swap"))
-						]
-					]
-				];
-		}
-
-		{
-			IDetailPropertyRow& EraseColorProp = TextureCategory.AddProperty(EraseColor);
-
-			TSharedPtr<SWidget> NameWidget;
-			TSharedPtr<SWidget> ValueWidget;
-
-			FDetailWidgetRow& EraseRow = EraseColorProp.CustomWidget(true);
-			EraseRow.ValueContent().MinDesiredWidth(250 - 16.f);
-			EraseColorProp.GetDefaultWidgets(NameWidget, ValueWidget, EraseRow, false);
-		}
-
-}
-
-FReply FTexturePaintingSettingsCustomization::OnSwapColorsClicked(TSharedRef<IPropertyHandle> PaintColor, TSharedRef<IPropertyHandle> EraseColor)
-{
-	UMeshTexturePaintingToolProperties* Settings = UMeshPaintMode::GetTextureToolProperties();
-	if (Settings)
-	{
-		FScopedTransaction Transaction(LOCTEXT("SwapColorsTransation", "Swap paint and erase colors"));
-		Settings->Modify();
-
-		FLinearColor TempPaintColor = Settings->PaintColor;
-		Settings->PaintColor = Settings->EraseColor;
-		Settings->EraseColor = TempPaintColor;
-		PaintColor->NotifyPostChange(EPropertyChangeType::ValueSet);
-		EraseColor->NotifyPostChange(EPropertyChangeType::ValueSet);
-	}
-	return FReply::Handled();
+	TextureWidget->AddSlot()
+	[
+		SNew(SObjectPropertyEntryBox)
+		.PropertyHandle(TextureProperty)
+		.AllowedClass(UTexture2D::StaticClass())
+		.OnShouldFilterAsset(FOnShouldFilterAsset::CreateUObject(Cast<UMeshTextureAssetPaintingTool>(UMeshPaintMode::GetMeshPaintMode()->GetToolManager()->GetActiveTool(EToolSide::Left)), &UMeshTextureAssetPaintingTool::ShouldFilterTextureAsset))
+		.OnObjectChanged(FOnSetObject::CreateUObject(Cast<UMeshTextureAssetPaintingTool>(UMeshPaintMode::GetMeshPaintMode()->GetToolManager()->GetActiveTool(EToolSide::Left)), &UMeshTextureAssetPaintingTool::PaintTextureChanged))
+		.DisplayUseSelected(false)
+		.ThumbnailPool(DetailLayout.GetThumbnailPool())
+	];
 }
 
 
diff --git a/Engine/Plugins/MeshPainting/Source/MeshPaintEditorMode/Private/BrushSettingsCustomization.h b/Engine/Plugins/MeshPainting/Source/MeshPaintEditorMode/Private/BrushSettingsCustomization.h
--- a/Engine/Plugins/MeshPainting/Source/MeshPaintEditorMode/Private/BrushSettingsCustomization.h
+++ b/Engine/Plugins/MeshPainting/Source/MeshPaintEditorMode/Private/BrushSettingsCustomization.h
@@ -8,27 +8,32 @@ class FPropertyRestriction;
 class FReply;
 class IPropertyHandle;
 
-class SHorizontalBox;
-class UMeshVertexPaintingToolProperties;
-
-
-class FVertexPaintingSettingsCustomization : public IDetailCustomization
+class FMeshPaintingSettingsCustomization : public IDetailCustomization
 {
 public:
 	static TSharedRef<IDetailCustomization> MakeInstance();
 
 	/** IPropertyTypeCustomization interface */
 	virtual void CustomizeDetails(IDetailLayoutBuilder& DetailLayout) override;
+
 protected:
-	/** Callback for when texture weight type changed so we can update restrictions */
-	void OnTextureWeightTypeChanged(TSharedRef<IPropertyHandle> WeightTypeProperty, TSharedRef<IPropertyHandle> PaintWeightProperty, TSharedRef<IPropertyHandle> EraseWeightProperty);
 	FReply OnSwapColorsClicked(TSharedRef<IPropertyHandle> PaintColor, TSharedRef<IPropertyHandle> EraseColor);
+};
 
+class FVertexPaintingSettingsCustomization : public FMeshPaintingSettingsCustomization
+{
+public:
+	static TSharedRef<IDetailCustomization> MakeInstance();
+
+	/** IPropertyTypeCustomization interface */
+	virtual void CustomizeDetails(IDetailLayoutBuilder& DetailLayout) override;
+
+protected:
 	/** Property restriction applied to blend paint enum dropdown box */
 	TSharedPtr<FPropertyRestriction> BlendPaintEnumRestriction;
 };
 
-class FColorPaintingSettingsCustomization : public FVertexPaintingSettingsCustomization
+class FVertexColorPaintingSettingsCustomization : public FVertexPaintingSettingsCustomization
 {
 public:
 	static TSharedRef<IDetailCustomization> MakeInstance();
@@ -37,30 +42,42 @@ public:
 	virtual void CustomizeDetails(IDetailLayoutBuilder& DetailLayout) override;
 };
 
-class FWeightPaintingSettingsCustomization : public FVertexPaintingSettingsCustomization
+class FVertexWeightPaintingSettingsCustomization : public FVertexPaintingSettingsCustomization
 {
 public:
 	static TSharedRef<IDetailCustomization> MakeInstance();
 
 	/** IPropertyTypeCustomization interface */
 	virtual void CustomizeDetails(IDetailLayoutBuilder& DetailLayout) override;
+
 protected:
+	/** Callback for when texture weight type changed so we can update restrictions */
 	void OnTextureWeightTypeChanged(TSharedRef<IPropertyHandle> WeightTypeProperty, TSharedRef<IPropertyHandle> PaintWeightProperty, TSharedRef<IPropertyHandle> EraseWeightProperty);
+};
 
-	/** Property restriction applied to blend paint enum dropdown box */
-	TSharedPtr<FPropertyRestriction> BlendPaintEnumRestriction;
+class FTexturePaintingSettingsCustomization : public FMeshPaintingSettingsCustomization
+{
+public:
+	static TSharedRef<IDetailCustomization> MakeInstance();
 
+	/** IPropertyTypeCustomization interface */
+	virtual void CustomizeDetails(IDetailLayoutBuilder& DetailLayout) override;
 };
 
-
-class FTexturePaintingSettingsCustomization : public IDetailCustomization
+class FTextureColorPaintingSettingsCustomization : public FTexturePaintingSettingsCustomization
 {
 public:
 	static TSharedRef<IDetailCustomization> MakeInstance();
+	
 	/** IPropertyTypeCustomization interface */
 	virtual void CustomizeDetails(IDetailLayoutBuilder& DetailLayout) override;
-	FReply OnSwapColorsClicked(TSharedRef<IPropertyHandle> PaintColor, TSharedRef<IPropertyHandle> EraseColor);
 };
 
+class FTextureAssetPaintingSettingsCustomization : public FTexturePaintingSettingsCustomization
+{
+public:
+	static TSharedRef<IDetailCustomization> MakeInstance();
 
-	
+	/** IPropertyTypeCustomization interface */
+	virtual void CustomizeDetails(IDetailLayoutBuilder& DetailLayout) override;
+};
\ No newline at end of file
diff --git a/Engine/Plugins/MeshPainting/Source/MeshPaintEditorMode/Private/MeshPaintEditorMode.cpp b/Engine/Plugins/MeshPainting/Source/MeshPaintEditorMode/Private/MeshPaintEditorMode.cpp
--- a/Engine/Plugins/MeshPainting/Source/MeshPaintEditorMode/Private/MeshPaintEditorMode.cpp
+++ b/Engine/Plugins/MeshPainting/Source/MeshPaintEditorMode/Private/MeshPaintEditorMode.cpp
@@ -50,9 +50,10 @@ void FMeshPaintEditorModeModule::Register()
 	RegisterGeometryAdapterFactory(MakeShareable(new FMeshPaintGeometryCollectionComponentAdapterFactory));
 
 	FPropertyEditorModule& PropertyModule = FModuleManager::LoadModuleChecked<FPropertyEditorModule>("PropertyEditor");
-	PropertyModule.RegisterCustomClassLayout("MeshColorPaintingToolProperties", FOnGetDetailCustomizationInstance::CreateStatic(&FColorPaintingSettingsCustomization::MakeInstance));
-	PropertyModule.RegisterCustomClassLayout("MeshWeightPaintingToolProperties", FOnGetDetailCustomizationInstance::CreateStatic(&FWeightPaintingSettingsCustomization::MakeInstance));
-	PropertyModule.RegisterCustomClassLayout("MeshTexturePaintingToolProperties", FOnGetDetailCustomizationInstance::CreateStatic(&FTexturePaintingSettingsCustomization::MakeInstance));
+	PropertyModule.RegisterCustomClassLayout("MeshVertexColorPaintingToolProperties", FOnGetDetailCustomizationInstance::CreateStatic(&FVertexColorPaintingSettingsCustomization::MakeInstance));
+	PropertyModule.RegisterCustomClassLayout("MeshVertexWeightPaintingToolProperties", FOnGetDetailCustomizationInstance::CreateStatic(&FVertexWeightPaintingSettingsCustomization::MakeInstance));
+	PropertyModule.RegisterCustomClassLayout("MeshTextureColorPaintingToolProperties", FOnGetDetailCustomizationInstance::CreateStatic(&FTextureColorPaintingSettingsCustomization::MakeInstance));
+	PropertyModule.RegisterCustomClassLayout("MeshTextureAssetPaintingToolProperties", FOnGetDetailCustomizationInstance::CreateStatic(&FTextureAssetPaintingSettingsCustomization::MakeInstance));
 }
 
 
diff --git a/Engine/Plugins/MeshPainting/Source/MeshPaintEditorMode/Private/MeshPaintMode.cpp b/Engine/Plugins/MeshPainting/Source/MeshPaintEditorMode/Private/MeshPaintMode.cpp
--- a/Engine/Plugins/MeshPainting/Source/MeshPaintEditorMode/Private/MeshPaintMode.cpp
+++ b/Engine/Plugins/MeshPainting/Source/MeshPaintEditorMode/Private/MeshPaintMode.cpp
@@ -28,86 +28,73 @@
 
 #include UE_INLINE_GENERATED_CPP_BY_NAME(MeshPaintMode)
 
-
 #define LOCTEXT_NAMESPACE "MeshPaintMode"
 
-FName UMeshPaintMode::MeshPaintMode_Color = FName(TEXT("Color"));
-FName UMeshPaintMode::MeshPaintMode_Weights = FName(TEXT("Weights"));
-FName UMeshPaintMode::MeshPaintMode_Texture = FName(TEXT("Texture"));
+FName UMeshPaintMode::MeshPaintMode_VertexColor = FName(TEXT("VertexColor"));
+FName UMeshPaintMode::MeshPaintMode_VertexWeights = FName(TEXT("VertexWeights"));
+FName UMeshPaintMode::MeshPaintMode_TextureColor = FName(TEXT("TextureColor"));
+FName UMeshPaintMode::MeshPaintMode_TextureAsset = FName(TEXT("Texture"));
 
 FString UMeshPaintMode::VertexSelectToolName = TEXT("VertexAdapterClickTool");
 FString UMeshPaintMode::TextureSelectToolName = TEXT("TextureAdapterClickTool");
-FString UMeshPaintMode::ColorPaintToolName = TEXT("ColorBrushTool");
-FString UMeshPaintMode::WeightPaintToolName = TEXT("WeightBrushTool");
-FString UMeshPaintMode::TexturePaintToolName = TEXT("TextureBrushTool");
 
+FString UMeshPaintMode::VertexColorPaintToolName = TEXT("VertexColorBrushTool");
+FString UMeshPaintMode::VertexWeightPaintToolName = TEXT("VertexWeightBrushTool");
+FString UMeshPaintMode::TextureColorPaintToolName = TEXT("TextureColorBrushTool");
+FString UMeshPaintMode::TextureAssetPaintToolName = TEXT("TextureBrushTool");
 
-UMeshVertexPaintingToolProperties* UMeshPaintMode::GetVertexToolProperties()
+
+template<typename T>
+static T* GetTypedToolProperties()
 {
 	if (GLevelEditorModeTools().GetActiveScriptableMode("MeshPaintMode")
 		&& GLevelEditorModeTools().GetActiveScriptableMode("MeshPaintMode")->GetToolManager()->GetActiveTool(EToolSide::Mouse))
 	{
 		TArray<UObject*> PropertyArray = GLevelEditorModeTools().GetActiveScriptableMode("MeshPaintMode")->GetToolManager()->GetActiveTool(EToolSide::Mouse)->GetToolProperties();
-		for(UObject* Property : PropertyArray)
+		for (UObject* Property : PropertyArray)
 		{
-			if (UMeshVertexPaintingToolProperties* VertexProperties = Cast<UMeshVertexPaintingToolProperties>(Property))
+			if (T* FoundProperty = Cast<T>(Property))
 			{
-				return VertexProperties;
+				return FoundProperty;
 			}
 		}
 	}
 	return nullptr;
 }
 
-UMeshColorPaintingToolProperties* UMeshPaintMode::GetColorToolProperties()
+UMeshPaintingToolProperties* UMeshPaintMode::GetToolProperties()
 {
-	if (GLevelEditorModeTools().GetActiveScriptableMode("MeshPaintMode")
-		&& GLevelEditorModeTools().GetActiveScriptableMode("MeshPaintMode")->GetToolManager()->GetActiveTool(EToolSide::Mouse))
-	{
-		TArray<UObject*> PropertyArray = GLevelEditorModeTools().GetActiveScriptableMode("MeshPaintMode")->GetToolManager()->GetActiveTool(EToolSide::Mouse)->GetToolProperties();
-		for (UObject* Property : PropertyArray)
-		{
-			if (UMeshColorPaintingToolProperties* ColorProperties = Cast<UMeshColorPaintingToolProperties>(Property))
-			{
-				return ColorProperties;
-			}
-		}
-	}
-	return nullptr;
+	return GetTypedToolProperties<UMeshPaintingToolProperties>();
 }
 
-UMeshWeightPaintingToolProperties* UMeshPaintMode::GetWeightToolProperties()
+UMeshVertexPaintingToolProperties* UMeshPaintMode::GetVertexToolProperties()
 {
-	if (GLevelEditorModeTools().GetActiveScriptableMode("MeshPaintMode")
-		&& GLevelEditorModeTools().GetActiveScriptableMode("MeshPaintMode")->GetToolManager()->GetActiveTool(EToolSide::Mouse))
-	{
-		TArray<UObject*> PropertyArray = GLevelEditorModeTools().GetActiveScriptableMode("MeshPaintMode")->GetToolManager()->GetActiveTool(EToolSide::Mouse)->GetToolProperties();
-		for (UObject* Property : PropertyArray)
-		{
-			if (UMeshWeightPaintingToolProperties* WeightProperties = Cast<UMeshWeightPaintingToolProperties>(Property))
-			{
-				return WeightProperties;
-			}
-		}
-	}
-	return nullptr;
+	return GetTypedToolProperties<UMeshVertexPaintingToolProperties>();
+}
+
+UMeshVertexColorPaintingToolProperties* UMeshPaintMode::GetVertexColorToolProperties()
+{
+	return GetTypedToolProperties<UMeshVertexColorPaintingToolProperties>();
+}
+
+UMeshVertexWeightPaintingToolProperties* UMeshPaintMode::GetVertexWeightToolProperties()
+{
+	return GetTypedToolProperties<UMeshVertexWeightPaintingToolProperties>();
 }
 
 UMeshTexturePaintingToolProperties* UMeshPaintMode::GetTextureToolProperties()
 {
-	if (GLevelEditorModeTools().GetActiveScriptableMode("MeshPaintMode")
-		&& GLevelEditorModeTools().GetActiveScriptableMode("MeshPaintMode")->GetToolManager()->GetActiveTool(EToolSide::Mouse))
-	{
-		TArray<UObject*> PropertyArray = GLevelEditorModeTools().GetActiveScriptableMode("MeshPaintMode")->GetToolManager()->GetActiveTool(EToolSide::Mouse)->GetToolProperties();
-		for (UObject* Property : PropertyArray)
-		{
-			if (UMeshTexturePaintingToolProperties* TextureProperties = Cast<UMeshTexturePaintingToolProperties>(Property))
-			{
-				return TextureProperties;
-			}
-		}
-	}
-	return nullptr;
+	return GetTypedToolProperties<UMeshTexturePaintingToolProperties>();
+}
+
+UMeshTextureColorPaintingToolProperties* UMeshPaintMode::GetTextureColorToolProperties()
+{
+	return GetTypedToolProperties<UMeshTextureColorPaintingToolProperties>();
+}
+
+UMeshTextureAssetPaintingToolProperties* UMeshPaintMode::GetTextureAssetToolProperties()
+{
+	return GetTypedToolProperties<UMeshTextureAssetPaintingToolProperties>();
 }
 
 UMeshPaintMode* UMeshPaintMode::GetMeshPaintMode()
@@ -115,6 +102,7 @@ UMeshPaintMode* UMeshPaintMode::GetMeshPaintMode()
 	return Cast<UMeshPaintMode>(GLevelEditorModeTools().GetActiveScriptableMode("MeshPaintMode"));
 }
 
+
 UMeshPaintMode::UMeshPaintMode()
 	: Super()
 {
@@ -132,6 +120,7 @@ UMeshPaintMode::UMeshPaintMode()
 void UMeshPaintMode::Enter()
 {
 	Super::Enter();
+
 	GEditor->OnEditorClose().AddUObject(this, &UMeshPaintMode::OnResetViewMode);
 	FCoreUObjectDelegates::OnObjectsReplaced.AddUObject(this, &UMeshPaintMode::OnObjectsReplaced);
 	ModeSettings = Cast<UMeshPaintModeSettings>(SettingsObject);
@@ -139,29 +128,31 @@ void UMeshPaintMode::Enter()
 	FMeshPaintEditorModeCommands ToolManagerCommands = FMeshPaintEditorModeCommands::Get();
 
 	UVertexAdapterClickToolBuilder* VertexClickToolBuilder = NewObject<UVertexAdapterClickToolBuilder>(this);
-	RegisterTool(ToolManagerCommands.VertexSelect, VertexSelectToolName, VertexClickToolBuilder);
+	RegisterTool(ToolManagerCommands.SelectVertex, VertexSelectToolName, VertexClickToolBuilder);
 
 	UTextureAdapterClickToolBuilder* TextureClickToolBuilder = NewObject<UTextureAdapterClickToolBuilder>(this);
-	RegisterTool(ToolManagerCommands.TextureSelect, TextureSelectToolName, TextureClickToolBuilder);
-
+	RegisterTool(ToolManagerCommands.SelectTexture, TextureSelectToolName, TextureClickToolBuilder);
 
 	UMeshColorPaintingToolBuilder* MeshColorPaintingToolBuilder = NewObject<UMeshColorPaintingToolBuilder>(this);
-	RegisterTool(ToolManagerCommands.ColorPaint, ColorPaintToolName, MeshColorPaintingToolBuilder);
+	RegisterTool(ToolManagerCommands.PaintVertexColor, VertexColorPaintToolName, MeshColorPaintingToolBuilder);
 
 	UMeshWeightPaintingToolBuilder* WeightPaintingToolBuilder = NewObject<UMeshWeightPaintingToolBuilder>(this);
-	RegisterTool(ToolManagerCommands.WeightPaint, WeightPaintToolName, WeightPaintingToolBuilder);
+	RegisterTool(ToolManagerCommands.PaintVertexWeight, VertexWeightPaintToolName, WeightPaintingToolBuilder);
 
-	UMeshTexturePaintingToolBuilder* TexturePaintingToolBuilder = NewObject<UMeshTexturePaintingToolBuilder>(this);
-	RegisterTool(ToolManagerCommands.TexturePaint, TexturePaintToolName, TexturePaintingToolBuilder);
-	UpdateSelectedMeshes();
+	UMeshTextureColorPaintingToolBuilder* MeshTextureColorPaintingToolBuilder = NewObject<UMeshTextureColorPaintingToolBuilder>(this);
+	RegisterTool(ToolManagerCommands.PaintTextureColor, TextureColorPaintToolName, MeshTextureColorPaintingToolBuilder);
 
+	UMeshTextureAssetPaintingToolBuilder* TextureAssetPaintingToolBuilder = NewObject<UMeshTextureAssetPaintingToolBuilder>(this);
+	RegisterTool(ToolManagerCommands.PaintTextureAsset, TextureAssetPaintToolName, TextureAssetPaintingToolBuilder);
+	
+	UpdateSelectedMeshes();
 
 	// Toolkit
 	PaletteChangedHandle = Toolkit->OnPaletteChanged().AddUObject(this, &UMeshPaintMode::UpdateOnPaletteChange);
 
 	// disable tool change tracking to activate default tool, and then switch to full undo/redo tracking mode
 	GetToolManager()->ConfigureChangeTrackingMode(EToolChangeTrackingMode::NoChangeTracking);
-	Toolkit->SetCurrentPalette(MeshPaintMode_Color);
+	Toolkit->SetCurrentPalette(MeshPaintMode_VertexColor);
 	GetToolManager()->ConfigureChangeTrackingMode(EToolChangeTrackingMode::UndoToExit);
 
 	FLevelEditorModule& LevelEditor = FModuleManager::GetModuleChecked<FLevelEditorModule>(FName(TEXT("LevelEditor")));
@@ -174,19 +165,15 @@ void UMeshPaintMode::Exit()
 	FCoreUObjectDelegates::OnObjectsReplaced.RemoveAll(this);
 	GEditor->OnEditorClose().RemoveAll(this);
 	OnResetViewMode();
+
 	const FMeshPaintEditorModeCommands& Commands = FMeshPaintEditorModeCommands::Get();
 	const TSharedRef<FUICommandList>& CommandList = Toolkit->GetToolkitCommands();
-	for (const TSharedPtr<const FUICommandInfo> Action : Commands.Commands[UMeshPaintMode::MeshPaintMode_Color])
+	for (auto It : Commands.Commands)
 	{
-		CommandList->UnmapAction(Action);
-	}
-	for (const TSharedPtr<const FUICommandInfo> Action : Commands.Commands[UMeshPaintMode::MeshPaintMode_Weights])
-	{
-		CommandList->UnmapAction(Action);
-	}
-	for (const TSharedPtr<const FUICommandInfo> Action : Commands.Commands[UMeshPaintMode::MeshPaintMode_Weights])
-	{
-		CommandList->UnmapAction(Action);
+		for (const TSharedPtr<const FUICommandInfo> Action : It.Value)
+		{
+			CommandList->UnmapAction(Action);
+		}
 	}
 
 	Super::Exit();
@@ -204,7 +191,7 @@ void UMeshPaintMode::CreateToolkit()
 
 void UMeshPaintMode::Tick(FEditorViewportClient* ViewportClient, float DeltaTime)
 {
-	EMeshPaintActiveMode CurrentActiveMode = EMeshPaintActiveMode::Color;
+	EMeshPaintActiveMode CurrentActiveMode = EMeshPaintActiveMode::VertexColor;
 
 	if (bRecacheVertexDataSize)
 	{
@@ -216,20 +203,25 @@ void UMeshPaintMode::Tick(FEditorViewportClient* ViewportClient, float DeltaTime
 	FString ActiveTool = GetToolManager()->GetActiveToolName(EToolSide::Mouse);
 	FName ActiveTab = Toolkit->GetCurrentPalette();
 	FName TargetTab = ActiveTab;
-	if (ActiveTool == TexturePaintToolName || ActiveTool == TextureSelectToolName)
+	if (ActiveTool == VertexColorPaintToolName)
 	{
-		TargetTab = MeshPaintMode_Texture;
-		CurrentActiveMode = EMeshPaintActiveMode::Texture;
+		TargetTab = MeshPaintMode_VertexColor;
+		CurrentActiveMode = EMeshPaintActiveMode::VertexColor;
 	}
-	else if (ActiveTool == ColorPaintToolName)
+	else if (ActiveTool == VertexWeightPaintToolName)
 	{
-		TargetTab = MeshPaintMode_Color;
-		CurrentActiveMode = EMeshPaintActiveMode::Color;
+		TargetTab = MeshPaintMode_VertexWeights;
+		CurrentActiveMode = EMeshPaintActiveMode::VertexWeights;
 	}
-	else if (ActiveTool == WeightPaintToolName)
+	else if (ActiveTool == TextureColorPaintToolName)
 	{
-		TargetTab = MeshPaintMode_Weights;
-		CurrentActiveMode = EMeshPaintActiveMode::Weights;
+		TargetTab = MeshPaintMode_TextureColor;
+		CurrentActiveMode = EMeshPaintActiveMode::TextureColor;
+	}
+	else if (ActiveTool == TextureAssetPaintToolName)
+	{
+		TargetTab = MeshPaintMode_TextureAsset;
+		CurrentActiveMode = EMeshPaintActiveMode::Texture;
 	}
 	if ( TargetTab != ActiveTab)
 	{
@@ -260,104 +252,89 @@ void UMeshPaintMode::BindCommands()
 	const FMeshPaintEditorModeCommands& Commands = FMeshPaintEditorModeCommands::Get();
 	const TSharedRef<FUICommandList>& CommandList = Toolkit->GetToolkitCommands();
 
-	CommandList->MapAction(Commands.SwitchForeAndBackgroundColor, FExecuteAction::CreateLambda([this]()
-	{
-		GEditor->GetEditorSubsystem<UMeshPaintModeSubsystem>()->SwapVertexColors();
-	}));
+	CommandList->MapAction(Commands.SwapColor, FUIAction(
+		FExecuteAction::CreateUObject(this, &UMeshPaintMode::SwapColors),
+		FCanExecuteAction::CreateUObject(this, &UMeshPaintMode::CanSwapColors)));
 
+ 	CommandList->MapAction(Commands.FillVertex,	FUIAction(
+ 		FExecuteAction::CreateUObject(this, &UMeshPaintMode::FillVertexColors),
+ 		FCanExecuteAction::CreateUObject(this, &UMeshPaintMode::CanFillVertexColors)));
 
-	CommandList->MapAction(Commands.CycleToNextLOD, FExecuteAction::CreateUObject(this, &UMeshPaintMode::CycleMeshLODs, 1));
-	CommandList->MapAction(Commands.CycleToPreviousLOD, FExecuteAction::CreateUObject(this, &UMeshPaintMode::CycleMeshLODs, -1));
-	UMeshPaintingSubsystem* MeshPaintingSubsystem = GEngine->GetEngineSubsystem<UMeshPaintingSubsystem>();
+ 	CommandList->MapAction(Commands.FillTexture, FUIAction(
+ 		FExecuteAction::CreateUObject(this, &UMeshPaintMode::FillTexture),
+ 		FCanExecuteAction::CreateUObject(this, &UMeshPaintMode::CanFillTexture)));
 
-	auto IsInValidPaintMode = [MeshPaintingSubsystem]() -> bool {
-		return (GetVertexToolProperties() != nullptr) && MeshPaintingSubsystem && MeshPaintingSubsystem->SelectionContainsValidAdapters();
-	};
-	CommandList->MapAction(Commands.Fill, 
-		FUIAction(FExecuteAction::CreateUObject(this, &UMeshPaintMode::FillWithVertexColor),
-		FCanExecuteAction::CreateLambda(IsInValidPaintMode)));
-
-	CommandList->MapAction(Commands.Propagate,
-		FUIAction(
-			FExecuteAction::CreateUObject(this, &UMeshPaintMode::PropagateVertexColorsToAsset),
-			FCanExecuteAction::CreateUObject(this, &UMeshPaintMode::CanPropagateVertexColors)
-		));
-
-	auto IsAValidMeshComponentSelected = [this, MeshPaintingSubsystem]() -> bool {
-		return (GetSelectedComponents<UMeshComponent>().Num() == 1) && MeshPaintingSubsystem && MeshPaintingSubsystem->SelectionContainsValidAdapters();
-	};
-	CommandList->MapAction(Commands.Import,
-		FUIAction(
-			FExecuteAction::CreateUObject(this, &UMeshPaintMode::ImportVertexColors),
-			FCanExecuteAction::CreateLambda(IsAValidMeshComponentSelected)
-		));
-
-	CommandList->MapAction(Commands.Save,
-		FUIAction(
-			FExecuteAction::CreateUObject(this, &UMeshPaintMode::SavePaintedAssets),
-			FCanExecuteAction::CreateUObject(this, &UMeshPaintMode::CanSaveMeshPackages)
-		));
-
-	CommandList->MapAction(Commands.Copy,
-		FUIAction(
-			FExecuteAction::CreateUObject(this, &UMeshPaintMode::CopyVertexColors),
-			FCanExecuteAction::CreateUObject(this, &UMeshPaintMode::CanCopyInstanceVertexColors)
-		));
-
-	CommandList->MapAction(Commands.Paste,
-		FUIAction(
-			FExecuteAction::CreateUObject(this, &UMeshPaintMode::PasteVertexColors),
-			FCanExecuteAction::CreateUObject(this, &UMeshPaintMode::CanPasteInstanceVertexColors)
-		));
-
-	CommandList->MapAction(Commands.Remove,
-		FUIAction(
-			FExecuteAction::CreateUObject(this, &UMeshPaintMode::RemoveVertexColors),
-			FCanExecuteAction::CreateUObject(this, &UMeshPaintMode::CanRemoveInstanceColors)
-		));
-
-	CommandList->MapAction(Commands.Fix,
-		FUIAction(
-			FExecuteAction::CreateUObject(this, &UMeshPaintMode::FixVertexColors),
-			FCanExecuteAction::CreateUObject(this, &UMeshPaintMode::DoesRequireVertexColorsFixup)
-		));
-
-	CommandList->MapAction(Commands.PropagateVertexColorsToLODs,
-		FUIAction(
-			FExecuteAction::CreateUObject(this, &UMeshPaintMode::PropagateVertexColorsToLODs),
-			FCanExecuteAction::CreateUObject(this, &UMeshPaintMode::CanPropagateVertexColorsToLODs)
-		));
-
-	CommandList->MapAction(Commands.NextTexture, 
-		FUIAction(
-			FExecuteAction::CreateUObject(this, &UMeshPaintMode::CycleTextures, 1),
-			FCanExecuteAction::CreateUObject(this, &UMeshPaintMode::CanCycleTextures)
-		));
-	CommandList->MapAction(Commands.PreviousTexture,
-		FUIAction(
-			FExecuteAction::CreateUObject(this, &UMeshPaintMode::CycleTextures, -1),
-			FCanExecuteAction::CreateUObject(this, &UMeshPaintMode::CanCycleTextures)
-		));
-	auto IsInValidTexturePaintMode = [this, MeshPaintingSubsystem]() -> bool {
-		if (MeshPaintingSubsystem)
-		{
-			FString ActiveTool = GetToolManager()->GetActiveToolName(EToolSide::Mouse);
-			if (ActiveTool == TexturePaintToolName)
-			{
-				return true;
-			}
-		}
+ 	CommandList->MapAction(Commands.Apply, FUIAction(
+ 		FExecuteAction::CreateUObject(this, &UMeshPaintMode::ApplyVertexColorsToAsset),
+ 		FCanExecuteAction::CreateUObject(this, &UMeshPaintMode::CanApplyVertexColorsToAsset)));
 
-		return false;
-	};
-	CommandList->MapAction(Commands.TextureFill,
-		FUIAction(FExecuteAction::CreateUObject(this, &UMeshPaintMode::FillWithTextureColor),
-			FCanExecuteAction::CreateLambda(IsInValidTexturePaintMode)));
-	CommandList->MapAction(Commands.SaveTexturePaint, FUIAction(FExecuteAction::CreateUObject(GEditor->GetEditorSubsystem<UMeshPaintModeSubsystem>(), &UMeshPaintModeSubsystem::SaveModifiedTextures), FCanExecuteAction::CreateUObject(GEditor->GetEditorSubsystem<UMeshPaintModeSubsystem>(), &UMeshPaintModeSubsystem::CanSaveModifiedTextures)));
+ 	CommandList->MapAction(Commands.Commit, FUIAction(
+ 		FExecuteAction::CreateUObject(this, &UMeshPaintMode::CommitTextureColorsToAsset),
+ 		FCanExecuteAction::CreateUObject(this, &UMeshPaintMode::CanCommitTextureColorsToAsset)));
+
+	CommandList->MapAction(Commands.Propagate, FUIAction(
+		FExecuteAction::CreateUObject(this, &UMeshPaintMode::PropagateVertexColorsToLODs),
+		FCanExecuteAction::CreateUObject(this, &UMeshPaintMode::CanPropagateVertexColorsToLODs)));
+
+	CommandList->MapAction(Commands.SaveVertex,	FUIAction(
+			FExecuteAction::CreateUObject(this, &UMeshPaintMode::SaveVertexColorsToAssets),
+			FCanExecuteAction::CreateUObject(this, &UMeshPaintMode::CanSaveVertexColorsToAssets)));
+
+ 	CommandList->MapAction(Commands.SaveTexture, FUIAction(
+ 		FExecuteAction::CreateUObject(this, &UMeshPaintMode::SaveTexturePackages),
+ 		FCanExecuteAction::CreateUObject(this, &UMeshPaintMode::CanSaveTexturePackages)));
+
+	CommandList->MapAction(Commands.Add, FUIAction(
+		FExecuteAction::CreateUObject(this, &UMeshPaintMode::AddMeshPaintTextures),
+		FCanExecuteAction::CreateUObject(this, &UMeshPaintMode::CanAddMeshPaintTextures)));
+
+	CommandList->MapAction(Commands.RemoveVertex, FUIAction(
+		FExecuteAction::CreateUObject(this, &UMeshPaintMode::RemoveInstanceVertexColors),
+		FCanExecuteAction::CreateUObject(this, &UMeshPaintMode::CanRemoveInstanceVertexColors)));
+
+	CommandList->MapAction(Commands.RemoveTexture, FUIAction(
+		FExecuteAction::CreateUObject(this, &UMeshPaintMode::RemoveMeshPaintTexture),
+		FCanExecuteAction::CreateUObject(this, &UMeshPaintMode::CanRemoveMeshPaintTextures)));
+
+ 	CommandList->MapAction(Commands.Copy, FUIAction(
+ 		FExecuteAction::CreateUObject(this, &UMeshPaintMode::Copy),
+ 		FCanExecuteAction::CreateUObject(this, &UMeshPaintMode::CanCopy)));
 
-	auto HasPaintChanges = [this]() -> bool { return GetNumberOfPendingPaintChanges() > 0; };
-	CommandList->MapAction(Commands.PropagateTexturePaint, FExecuteAction::CreateUObject(this, &UMeshPaintMode::CommitAllPaintedTextures), FCanExecuteAction::CreateLambda(HasPaintChanges));
+	CommandList->MapAction(Commands.Paste, FUIAction(
+		FExecuteAction::CreateUObject(this, &UMeshPaintMode::Paste),
+		FCanExecuteAction::CreateUObject(this, &UMeshPaintMode::CanPaste)));
 
+	CommandList->MapAction(Commands.Import, FUIAction(
+		FExecuteAction::CreateUObject(this, &UMeshPaintMode::ImportVertexColorsFromFile),
+		FCanExecuteAction::CreateUObject(this, &UMeshPaintMode::CanImportVertexColorsFromFile)));
+
+	CommandList->MapAction(Commands.GetTextureColors, FUIAction(
+		FExecuteAction::CreateUObject(this, &UMeshPaintMode::ImportVertexColorsFromMeshPaintTexture),
+		FCanExecuteAction::CreateUObject(this, &UMeshPaintMode::CanImportVertexColorsFromMeshPaintTexture)));
+
+	CommandList->MapAction(Commands.GetVertexColors, FUIAction(
+		FExecuteAction::CreateUObject(this, &UMeshPaintMode::ImportMeshPaintTextureFromVertexColors),
+		FCanExecuteAction::CreateUObject(this, &UMeshPaintMode::CanImportMeshPaintTextureFromVertexColors)));
+
+	CommandList->MapAction(Commands.Fix, FUIAction(
+		FExecuteAction::CreateUObject(this, &UMeshPaintMode::FixVertexColors),
+		FCanExecuteAction::CreateUObject(this, &UMeshPaintMode::CanFixVertexColors)));
+
+	CommandList->MapAction(Commands.PreviousLOD, FUIAction(
+		FExecuteAction::CreateUObject(this, &UMeshPaintMode::CycleMeshLODs, -1),
+		FCanExecuteAction::CreateUObject(this, &UMeshPaintMode::CanCycleMeshLODs)));
+
+	CommandList->MapAction(Commands.NextLOD, FUIAction(
+		FExecuteAction::CreateUObject(this, &UMeshPaintMode::CycleMeshLODs, 1),
+		FCanExecuteAction::CreateUObject(this, &UMeshPaintMode::CanCycleMeshLODs)));
+
+	CommandList->MapAction(Commands.PreviousTexture, FUIAction(
+		FExecuteAction::CreateUObject(this, &UMeshPaintMode::CycleTextures, -1),
+		FCanExecuteAction::CreateUObject(this, &UMeshPaintMode::CanCycleTextures)));
+
+	CommandList->MapAction(Commands.NextTexture, FUIAction(
+		FExecuteAction::CreateUObject(this, &UMeshPaintMode::CycleTextures, 1),
+		FCanExecuteAction::CreateUObject(this, &UMeshPaintMode::CanCycleTextures)));
 }
 
 void UMeshPaintMode::OnToolStarted(UInteractiveToolManager* Manager, UInteractiveTool* Tool)
@@ -372,7 +349,7 @@ void UMeshPaintMode::OnToolStarted(UInteractiveToolManager* Manager, UInteractiv
 
 void UMeshPaintMode::OnVertexPaintFinished()
 {
-	if (UMeshColorPaintingToolProperties* ColorPaintingToolProperties = UMeshPaintMode::GetColorToolProperties())
+	if (UMeshVertexColorPaintingToolProperties* ColorPaintingToolProperties = UMeshPaintMode::GetVertexColorToolProperties())
 	{
 		if (!ColorPaintingToolProperties->bPaintOnSpecificLOD)
 		{
@@ -405,33 +382,36 @@ void UMeshPaintMode::UpdateSelectedMeshes()
 {
 	if (UMeshPaintingSubsystem* MeshPaintingSubsystem = GEngine->GetEngineSubsystem<UMeshPaintingSubsystem>())
 	{
-		FName PaletteName = Toolkit->GetCurrentPalette();
-
 		MeshPaintingSubsystem->ResetState();
 		TArray<UMeshComponent*> CurrentMeshComponents = GetSelectedComponents<UMeshComponent>();
 		MeshPaintingSubsystem->AddSelectedMeshComponents(CurrentMeshComponents);
 		MeshPaintingSubsystem->bNeedsRecache = true;
 
 		// Check the current selection for a material that supports texture paint
-		if (PaletteName == UMeshPaintMode::MeshPaintMode_Texture)
+		FName PaletteName = Toolkit->GetCurrentPalette();
+		if (PaletteName == UMeshPaintMode::MeshPaintMode_TextureColor || PaletteName == UMeshPaintMode::MeshPaintMode_TextureAsset)
 		{
-			CheckSelectionForTexturePaintCompat(CurrentMeshComponents);
+			UpdateToolForSelection(CurrentMeshComponents);
 		}
 	}
 	
 	bRecacheVertexDataSize = true;
 }
 
-void UMeshPaintMode::CheckSelectionForTexturePaintCompat(const TArray<UMeshComponent*>& CurrentMeshComponents)
+void UMeshPaintMode::UpdateToolForSelection(const TArray<UMeshComponent*>& CurrentMeshComponents)
 {
+	FName PaletteName = Toolkit->GetCurrentPalette();
+	const bool bIsTextureColorPainting = PaletteName == UMeshPaintMode::MeshPaintMode_TextureColor;
+
 	bool bCurrentSelectionSupportsTexturePaint = false;
-	if (CurrentMeshComponents.Num())
+	
+	// Texture Color painting supports multiselect. But other modes require a single mesh selection.
+	if (CurrentMeshComponents.Num() == 1 || bIsTextureColorPainting)
 	{
-		UActorComponent* ActorComponent = CurrentMeshComponents[0];
 		TArray<FPaintableTexture> PaintableTextures;
-		if (ActorComponent)
+
+		for (UMeshComponent* MeshComponent : CurrentMeshComponents)
 		{
-			UMeshComponent* MeshComponent = Cast<UMeshComponent>(ActorComponent);
 			if (MeshComponent)
 			{
 				TSharedPtr<IMeshPaintComponentAdapter> MeshAdapter = GEngine->GetEngineSubsystem<UMeshPaintingSubsystem>()->GetAdapterForComponent(MeshComponent);
@@ -443,9 +423,16 @@ void UMeshPaintMode::CheckSelectionForTexturePaintCompat(const TArray<UMeshCompo
 						GEngine->GetEngineSubsystem<UMeshPaintingSubsystem>()->AddToComponentToAdapterMap(MeshComponent, MeshAdapter);
 					}
 				}
-				UTexturePaintToolset::RetrieveTexturesForComponent(MeshComponent, MeshAdapter.Get(), PaintableTextures);
+
+				int32 DummyDefaultIndex = INDEX_NONE;
+				UTexturePaintToolset::RetrieveTexturesForComponent(MeshComponent, MeshAdapter.Get(), DummyDefaultIndex, PaintableTextures);
 			}
 		}
+
+		// PaintableTextures are collected for both TextureColor painting (MeshPaintTextures on components) and Texture painting (Textures ref'd in materials).
+		// Filter these according to the current texture painting mode here.
+ 		PaintableTextures.RemoveAll([bIsTextureColorPainting](FPaintableTexture const& PaintableTexture) { return bIsTextureColorPainting != PaintableTexture.bIsMeshTexture; });
+
 		bCurrentSelectionSupportsTexturePaint = PaintableTextures.Num() > 0;
 	}
 
@@ -462,12 +449,12 @@ void UMeshPaintMode::CheckSelectionForTexturePaintCompat(const TArray<UMeshCompo
 
 void UMeshPaintMode::UpdateOnMaterialChange(bool bInvalidateHitProxies)
 {
-	FName PaletteName = Toolkit->GetCurrentPalette();
 	// Check the current selection for a material that supports texture paint
-	if (PaletteName == UMeshPaintMode::MeshPaintMode_Texture)
+	FName PaletteName = Toolkit->GetCurrentPalette();
+	if (PaletteName == UMeshPaintMode::MeshPaintMode_TextureColor || PaletteName == UMeshPaintMode::MeshPaintMode_TextureAsset)
 	{
 		const TArray<UMeshComponent*> CurrentMeshComponents = GetSelectedComponents<UMeshComponent>();
-		CheckSelectionForTexturePaintCompat(CurrentMeshComponents);
+		UpdateToolForSelection(CurrentMeshComponents);
 	}
 }
 
@@ -481,7 +468,7 @@ void UMeshPaintMode::OnObjectsReplaced(const TMap<UObject*, UObject*>& InOldToNe
 	}
 }
 
-void UMeshPaintMode::FillWithVertexColor()
+void UMeshPaintMode::FillVertexColors()
 {
 	FScopedTransaction Transaction(LOCTEXT("LevelMeshPainter_TransactionFillInstColors", "Filling Per-Instance Vertex Colors"));
 	const TArray<UMeshComponent*> MeshComponents = GetSelectedComponents<UMeshComponent>();
@@ -492,9 +479,9 @@ void UMeshPaintMode::FillWithVertexColor()
 
 	if (GetToolManager()->GetActiveTool(EToolSide::Mouse)->IsA<UMeshWeightPaintingTool>())
 	{
-		FillColor = MeshPaintingSubsystem->GenerateColorForTextureWeight((int32)GetWeightToolProperties()->TextureWeightType, (int32)GetWeightToolProperties()->PaintTextureWeightIndex).ToFColor(bConvertSRGB);
+		FillColor = MeshPaintingSubsystem->GenerateColorForTextureWeight((int32)GetVertexWeightToolProperties()->TextureWeightType, (int32)GetVertexWeightToolProperties()->PaintTextureWeightIndex).ToFColor(bConvertSRGB);
 	}
-	else if (UMeshColorPaintingToolProperties* ColorProperties = GetColorToolProperties())
+	else if (UMeshVertexColorPaintingToolProperties* ColorProperties = GetVertexColorToolProperties())
 	{
 		FillColor = ColorProperties->PaintColor.ToFColor(bConvertSRGB);
 		MaskColor.R = ColorProperties->bWriteRed ? 255 : 0;
@@ -517,7 +504,7 @@ void UMeshPaintMode::FillWithVertexColor()
 			MeshAdapter->PreEdit();
 		}
 
-		UMeshColorPaintingToolProperties* ColorProperties = UMeshPaintMode::GetColorToolProperties();
+		UMeshVertexColorPaintingToolProperties* ColorProperties = UMeshPaintMode::GetVertexColorToolProperties();
 		
 		const bool bPaintOnSpecificLOD = ColorProperties ? ColorProperties->bPaintOnSpecificLOD : false;
 
@@ -550,7 +537,7 @@ void UMeshPaintMode::FillWithVertexColor()
 }
 
 
-void UMeshPaintMode::PropagateVertexColorsToAsset()
+void UMeshPaintMode::ApplyVertexColorsToAsset()
 {
 	const TArray<UStaticMeshComponent*> StaticMeshComponents = GetSelectedComponents<UStaticMeshComponent>();
 	FSuppressableWarningDialog::FSetupInfo SetupInfo(LOCTEXT("PushInstanceVertexColorsPrompt_Message", "This operation copies vertex colors from LOD 0 of the selected instance to all LODs of the source asset, overwriting any existing vertex colors.\n\nThis change will also propagate to all other instances of the same asset that do not have custom vertex colors."),
@@ -571,7 +558,7 @@ void UMeshPaintMode::PropagateVertexColorsToAsset()
 	}
 }
 
-bool UMeshPaintMode::CanPropagateVertexColors() const
+bool UMeshPaintMode::CanApplyVertexColorsToAsset() const
 {
 	// Check whether or not our selected Static Mesh Components contain instance based vertex colors (only these can be propagated to the base mesh)
 	int32 NumInstanceVertexColorBytes = 0;
@@ -579,10 +566,9 @@ bool UMeshPaintMode::CanPropagateVertexColors() const
 	TArray<UStaticMesh*> StaticMeshes;
 	TArray<UStaticMeshComponent*> StaticMeshComponents = GetSelectedComponents<UStaticMeshComponent>();
 	return GEditor->GetEditorSubsystem<UMeshPaintModeSubsystem>()->CanPropagateVertexColors(StaticMeshComponents, StaticMeshes, NumInstanceVertexColorBytes);
-
 }
 
-void UMeshPaintMode::ImportVertexColors()
+void UMeshPaintMode::ImportVertexColorsFromFile()
 {
 	const TArray<UMeshComponent*> MeshComponents = GetSelectedComponents<UMeshComponent>();
 	if (MeshComponents.Num() == 1)
@@ -593,7 +579,7 @@ void UMeshPaintMode::ImportVertexColors()
 	}
 }
 
-void UMeshPaintMode::SavePaintedAssets()
+void UMeshPaintMode::SaveVertexColorsToAssets()
 {
 	const TArray<UStaticMeshComponent*> StaticMeshComponents = GetSelectedComponents<UStaticMeshComponent>();
 	const TArray<USkeletalMeshComponent*> SkeletalMeshComponents = GetSelectedComponents<USkeletalMeshComponent>();
@@ -622,7 +608,7 @@ void UMeshPaintMode::SavePaintedAssets()
 	}
 }
 
-bool UMeshPaintMode::CanSaveMeshPackages() const
+bool UMeshPaintMode::CanSaveVertexColorsToAssets() const
 {
 	// Check whether or not any of our selected mesh components contain mesh objects which require saving
 	TArray<UMeshComponent*> Components = GetSelectedComponents<UMeshComponent>();
@@ -651,11 +637,11 @@ bool UMeshPaintMode::CanSaveMeshPackages() const
 	return bValid;
 }
 
-bool UMeshPaintMode::CanRemoveInstanceColors() const
+bool UMeshPaintMode::CanRemoveInstanceVertexColors() const
 {
 	const TArray<UStaticMeshComponent*> StaticMeshComponents = GetSelectedComponents<UStaticMeshComponent>();
 	int32 PaintingMeshLODIndex = 0;
-	if (UMeshColorPaintingToolProperties* ColorProperties = UMeshPaintMode::GetColorToolProperties())
+	if (UMeshVertexColorPaintingToolProperties* ColorProperties = UMeshPaintMode::GetVertexColorToolProperties())
 	{
 		PaintingMeshLODIndex = ColorProperties->bPaintOnSpecificLOD ? ColorProperties->LODIndex : 0;
 	}
@@ -688,7 +674,7 @@ bool UMeshPaintMode::CanCopyInstanceVertexColors() const
 {
 	const TArray<UStaticMeshComponent*> StaticMeshComponents = GetSelectedComponents<UStaticMeshComponent>();
 	int32 PaintingMeshLODIndex = 0;
-	if (UMeshColorPaintingToolProperties* ColorProperties = UMeshPaintMode::GetColorToolProperties())
+	if (UMeshVertexColorPaintingToolProperties* ColorProperties = UMeshPaintMode::GetVertexColorToolProperties())
 	{
 		PaintingMeshLODIndex = ColorProperties->bPaintOnSpecificLOD ? ColorProperties->LODIndex : 0;
 	}
@@ -700,7 +686,7 @@ bool UMeshPaintMode::CanCopyInstanceVertexColors() const
 bool UMeshPaintMode::CanPropagateVertexColorsToLODs() const
 {
 	bool bPaintOnSpecificLOD = false;
-	if (UMeshColorPaintingToolProperties* ColorProperties = UMeshPaintMode::GetColorToolProperties())
+	if (UMeshVertexColorPaintingToolProperties* ColorProperties = UMeshPaintMode::GetVertexColorToolProperties())
 	{
 		bPaintOnSpecificLOD = ColorProperties ? ColorProperties->bPaintOnSpecificLOD : false;
 	}
@@ -709,16 +695,15 @@ bool UMeshPaintMode::CanPropagateVertexColorsToLODs() const
 	return bSelectionContainsPerLODColors || !bPaintOnSpecificLOD;
 }
 
-void UMeshPaintMode::CopyVertexColors()
+void UMeshPaintMode::CopyInstanceVertexColors()
 {
 	const TArray<UStaticMeshComponent*> StaticMeshComponents = GetSelectedComponents<UStaticMeshComponent>();
 	TArray<FPerComponentVertexColorData> CopiedColorsByComponent;
 	GEditor->GetEditorSubsystem<UMeshPaintModeSubsystem>()->CopyVertexColors(StaticMeshComponents, CopiedColorsByComponent);
 	GEngine->GetEngineSubsystem<UMeshPaintingSubsystem>()->SetCopiedColorsByComponent(CopiedColorsByComponent);
-
 }
 
-void UMeshPaintMode::PasteVertexColors()
+void UMeshPaintMode::PasteInstanceVertexColors()
 {
 	FScopedTransaction Transaction(LOCTEXT("LevelMeshPainter_TransactionPasteInstColors", "Pasting Per-Instance Vertex Colors"));
 	const TArray<UStaticMeshComponent*> StaticMeshComponents = GetSelectedComponents<UStaticMeshComponent>();
@@ -737,7 +722,7 @@ void UMeshPaintMode::FixVertexColors()
 	}
 }
 
-bool UMeshPaintMode::DoesRequireVertexColorsFixup() const
+bool UMeshPaintMode::CanFixVertexColors() const
 {
 	const TArray<UStaticMeshComponent*> StaticMeshComponents = GetSelectedComponents<UStaticMeshComponent>();
 	bool bAnyMeshNeedsFixing = false;
@@ -750,7 +735,7 @@ bool UMeshPaintMode::DoesRequireVertexColorsFixup() const
 	return bAnyMeshNeedsFixing;
 }
 
-void UMeshPaintMode::RemoveVertexColors()
+void UMeshPaintMode::RemoveInstanceVertexColors()
 {
 	FScopedTransaction Transaction(LOCTEXT("LevelMeshPainter_TransactionRemoveInstColors", "Removing Per-Instance Vertex Colors"));
 	const TArray<UStaticMeshComponent*> StaticMeshComponents = GetSelectedComponents<UStaticMeshComponent>();
@@ -781,7 +766,7 @@ void UMeshPaintMode::PropagateVertexColorsToLODs()
 		SetupInfo.CheckBoxText = LOCTEXT("LooseLowersLODsVertexColorsPrompt_CheckBoxText", "Always overwrite LODs without prompting");
 
 		FSuppressableWarningDialog LooseLowersLODsVertexColorsWarning(SetupInfo);
-	
+
 		// Prompt the user to see if they really want to propagate the base lod vert colors to the lowers LODs.
 		if (LooseLowersLODsVertexColorsWarning.ShowModal() == FSuppressableWarningDialog::Cancel)
 		{
@@ -891,7 +876,7 @@ void UMeshPaintMode::CycleMeshLODs(int32 Direction)
 
 void UMeshPaintMode::CycleTextures(int32 Direction)
 {
-	if (UMeshTexturePaintingTool* TexturePaintingTool = Cast<UMeshTexturePaintingTool>(GetToolManager()->GetActiveTool(EToolSide::Left)))
+	if (UMeshTextureAssetPaintingTool* TexturePaintingTool = Cast<UMeshTextureAssetPaintingTool>(GetToolManager()->GetActiveTool(EToolSide::Left)))
 	{
 		FScopedTransaction Transaction(LOCTEXT("LevelMeshPainter_CycleTexture", "Changed Current Texture"));
 		TexturePaintingTool->CycleTextures(Direction);
@@ -900,11 +885,11 @@ void UMeshPaintMode::CycleTextures(int32 Direction)
 
 bool UMeshPaintMode::CanCycleTextures() const
 {
-	UMeshTexturePaintingTool* TexturePaintingTool = Cast<UMeshTexturePaintingTool>(GetToolManager()->GetActiveTool(EToolSide::Left));
+	UMeshTextureAssetPaintingTool* TexturePaintingTool = Cast<UMeshTextureAssetPaintingTool>(GetToolManager()->GetActiveTool(EToolSide::Left));
 	return TexturePaintingTool != nullptr;
 }
 
-void UMeshPaintMode::CommitAllPaintedTextures()
+void UMeshPaintMode::CommitTextureColorsToAsset()
 {
 	if (UMeshTexturePaintingTool* TexturePaintingTool = Cast<UMeshTexturePaintingTool>(GetToolManager()->GetActiveTool(EToolSide::Left)))
 	{
@@ -912,28 +897,24 @@ void UMeshPaintMode::CommitAllPaintedTextures()
 	}
 }
 
-int32 UMeshPaintMode::GetNumberOfPendingPaintChanges()
+bool UMeshPaintMode::CanCommitTextureColorsToAsset() const
 {
 	if (UMeshTexturePaintingTool* TexturePaintingTool = Cast<UMeshTexturePaintingTool>(GetToolManager()->GetActiveTool(EToolSide::Left)))
 	{
-		return TexturePaintingTool->GetNumberOfPendingPaintChanges();
+		return TexturePaintingTool->GetNumberOfPendingPaintChanges() > 0;
 	}
-	return 0;
+	return false;
 }
 
 
 void UMeshPaintMode::ActivateDefaultTool()
 {
 	FName PaletteName = Toolkit->GetCurrentPalette();
-	if (PaletteName == UMeshPaintMode::MeshPaintMode_Color)
-	{
-		GetInteractiveToolsContext()->StartTool(VertexSelectToolName);
-	}
-	if (PaletteName == UMeshPaintMode::MeshPaintMode_Weights)
+	if (PaletteName == UMeshPaintMode::MeshPaintMode_VertexColor || PaletteName == UMeshPaintMode::MeshPaintMode_VertexWeights)
 	{
 		GetInteractiveToolsContext()->StartTool(VertexSelectToolName);
 	}
-	if (PaletteName == UMeshPaintMode::MeshPaintMode_Texture)
+	if (PaletteName == UMeshPaintMode::MeshPaintMode_TextureColor || PaletteName == UMeshPaintMode::MeshPaintMode_TextureAsset)
 	{
 		GetInteractiveToolsContext()->StartTool(TextureSelectToolName);
 	}
@@ -943,31 +924,46 @@ void UMeshPaintMode::UpdateOnPaletteChange(FName NewPaletteName)
 {
 	UpdateSelectedMeshes();
 
-	FString ActiveTool = GetToolManager()->GetActiveToolName(EToolSide::Mouse);
-	bool bInAnyPaintTool = (ActiveTool == WeightPaintToolName || ActiveTool == TexturePaintToolName || ActiveTool == ColorPaintToolName);
-
-	// figure out which tool we would like to be in based on currently-active tool
-	FString SwitchToTool;
-	if (NewPaletteName == UMeshPaintMode::MeshPaintMode_Color)
+	FString SwitchToToolPaint;
+	FString SwitchToToolSelect;
+	if (NewPaletteName == UMeshPaintMode::MeshPaintMode_VertexColor)
 	{
-		SwitchToTool = (bInAnyPaintTool) ? ColorPaintToolName : VertexSelectToolName;
+		SwitchToToolPaint = VertexColorPaintToolName;
+		SwitchToToolSelect = VertexSelectToolName;
 	}
-	else if (NewPaletteName == UMeshPaintMode::MeshPaintMode_Weights)
+	else if (NewPaletteName == UMeshPaintMode::MeshPaintMode_VertexWeights)
 	{
-		SwitchToTool = (bInAnyPaintTool) ? WeightPaintToolName: VertexSelectToolName;
+		SwitchToToolPaint = VertexWeightPaintToolName;
+		SwitchToToolSelect = VertexSelectToolName;
 	}
-	else if (NewPaletteName == UMeshPaintMode::MeshPaintMode_Texture)
+	else if (NewPaletteName == UMeshPaintMode::MeshPaintMode_TextureColor)
 	{
-		SwitchToTool = (bInAnyPaintTool) ? TexturePaintToolName : TextureSelectToolName;
+		SwitchToToolPaint = TextureColorPaintToolName;
+		SwitchToToolSelect = TextureSelectToolName;
 	}
-	// change to new tool if it is different
-	if (SwitchToTool.IsEmpty() == false && SwitchToTool != ActiveTool)
+	else if (NewPaletteName == UMeshPaintMode::MeshPaintMode_TextureAsset)
 	{
-		GetInteractiveToolsContext()->StartTool(SwitchToTool);
+		SwitchToToolPaint = TextureAssetPaintToolName;
+		SwitchToToolSelect = TextureSelectToolName;
+	}
+
+	if (!SwitchToToolPaint.IsEmpty())
+	{
+		// Figure out which tool we would like to be in based on currently-active tool
+		const FString ActiveTool = GetToolManager()->GetActiveToolName(EToolSide::Mouse);
+		const bool bInAnyPaintTool = (ActiveTool == VertexColorPaintToolName || ActiveTool == VertexWeightPaintToolName || ActiveTool == TextureColorPaintToolName || ActiveTool == TextureAssetPaintToolName);
+		const bool bUsePaintTool = bInAnyPaintTool && GetInteractiveToolsContext()->CanStartTool(SwitchToToolPaint);
+		const FString SwitchToTool = bUsePaintTool ? SwitchToToolPaint : SwitchToToolSelect;
+
+		// Change to new tool if it is different
+		if (SwitchToTool != ActiveTool)
+		{
+			GetInteractiveToolsContext()->StartTool(SwitchToTool);
+		}
 	}
 }
 
-void UMeshPaintMode::FillWithTextureColor()
+void UMeshPaintMode::FillTexture()
 {
 	if (UMeshTexturePaintingTool* TexturePaintingTool = Cast<UMeshTexturePaintingTool>(GetToolManager()->GetActiveTool(EToolSide::Left)))
 	{
@@ -985,11 +981,292 @@ void UMeshPaintMode::OnResetViewMode()
 			continue;
 		}
 
-		GEditor->GetEditorSubsystem<UMeshPaintModeSubsystem>()->SetViewportColorMode(EMeshPaintActiveMode::Color, EMeshPaintDataColorViewMode::Normal, ViewportClient);
+		GEditor->GetEditorSubsystem<UMeshPaintModeSubsystem>()->SetViewportColorMode(EMeshPaintActiveMode::VertexColor, EMeshPaintDataColorViewMode::Normal, ViewportClient);
 	}
 }
 
+bool UMeshPaintMode::IsInSelectTool() const
+{
+	FString ActiveTool = GetToolManager()->GetActiveToolName(EToolSide::Mouse);
+	return ActiveTool == VertexSelectToolName || ActiveTool == TextureSelectToolName;
+}
+
+bool UMeshPaintMode::IsInPaintTool() const
+{
+	FString ActiveTool = GetToolManager()->GetActiveToolName(EToolSide::Mouse);
+	return ActiveTool == VertexColorPaintToolName || ActiveTool == VertexWeightPaintToolName || ActiveTool == TextureColorPaintToolName || ActiveTool == TextureAssetPaintToolName;
+}
 
-#undef LOCTEXT_NAMESPACE
+void UMeshPaintMode::SwapColors()
+{
+	GEditor->GetEditorSubsystem<UMeshPaintModeSubsystem>()->SwapColors();
+}
+
+bool UMeshPaintMode::CanSwapColors() const
+{
+	return IsInPaintTool();
+}
+
+bool UMeshPaintMode::CanFillVertexColors() const
+{
+	return IsInPaintTool();
+}
 
+bool UMeshPaintMode::CanFillTexture() const
+{
+	return IsInPaintTool();
+}
+
+void UMeshPaintMode::SaveTexturePackages()
+{
+	TArray<UObject*> TexturesToSave;
+	if (UMeshTexturePaintingTool* TexturePaintingTool = Cast<UMeshTexturePaintingTool>(GetToolManager()->GetActiveTool(EToolSide::Left)))
+	{
+		TexturePaintingTool->GetModifiedTexturesToSave(TexturesToSave);
+	}
+	if (TexturesToSave.Num())
+	{
+		UPackageTools::SavePackagesForObjects(TexturesToSave);
+	}
+}
 
+bool UMeshPaintMode::CanSaveTexturePackages() const
+{
+	TArray<UObject*> TexturesToSave;
+	if (UMeshTexturePaintingTool* TexturePaintingTool = Cast<UMeshTexturePaintingTool>(GetToolManager()->GetActiveTool(EToolSide::Left)))
+	{
+		TexturePaintingTool->GetModifiedTexturesToSave(TexturesToSave);
+	}
+	return TexturesToSave.Num() > 0;
+}
+
+void UMeshPaintMode::AddMeshPaintTextures()
+{
+	FScopedTransaction Transaction(LOCTEXT("LevelMeshPainter_TransactionAddMeshPaintTexture", "Creating Mesh Paint Texture"));
+	const TArray<UStaticMeshComponent*> StaticMeshComponents = GetSelectedComponents<UStaticMeshComponent>();
+	for (UStaticMeshComponent* Component : StaticMeshComponents)
+	{
+		GEngine->GetEngineSubsystem<UMeshPaintingSubsystem>()->CreateComponentMeshPaintTexture(Component);
+	}
+
+	// The selection/paint status will change after add.
+	const TArray<UMeshComponent*> CurrentMeshComponents = GetSelectedComponents<UMeshComponent>();
+	UpdateToolForSelection(CurrentMeshComponents);
+}
+
+bool UMeshPaintMode::CanAddMeshPaintTextures() const
+{
+	if (IsInSelectTool())
+	{
+		const TArray<UStaticMeshComponent*> StaticMeshComponents = GetSelectedComponents<UStaticMeshComponent>();
+		for (UStaticMeshComponent* Component : StaticMeshComponents)
+		{
+			if (Component->GetMeshPaintTexture() == nullptr)
+			{
+				return true;
+			}
+		}
+	}
+	return false;
+}
+
+void UMeshPaintMode::RemoveMeshPaintTexture()
+{
+	FScopedTransaction Transaction(LOCTEXT("LevelMeshPainter_TransactionRemoveMeshPaintTexture", "Removing Mesh Paint Texture"));
+	const TArray<UStaticMeshComponent*> StaticMeshComponents = GetSelectedComponents<UStaticMeshComponent>();
+	for (UStaticMeshComponent* Component : StaticMeshComponents)
+	{
+		GEngine->GetEngineSubsystem<UMeshPaintingSubsystem>()->RemoveComponentMeshPaintTexture(Component);
+	}
+
+	// The selection/paint status will change after removal.
+	const TArray<UMeshComponent*> CurrentMeshComponents = GetSelectedComponents<UMeshComponent>();
+	UpdateToolForSelection(CurrentMeshComponents);
+}
+
+bool UMeshPaintMode::CanRemoveMeshPaintTextures() const
+{
+	if (IsInSelectTool())
+	{
+		const TArray<UStaticMeshComponent*> StaticMeshComponents = GetSelectedComponents<UStaticMeshComponent>();
+		for (UStaticMeshComponent* Component : StaticMeshComponents)
+		{
+			if (Component->GetMeshPaintTexture() != nullptr)
+			{
+				return true;
+			}
+		}
+	}
+	return false;
+}
+
+void UMeshPaintMode::CopyMeshPaintTexture()
+{
+	const TArray<UStaticMeshComponent*> StaticMeshComponents = GetSelectedComponents<UStaticMeshComponent>();
+	if (StaticMeshComponents.Num() > 0)
+	{
+		GEngine->GetEngineSubsystem<UMeshPaintingSubsystem>()->SetCopiedTexture(StaticMeshComponents[0]->GetMeshPaintTexture());
+	}
+}
+
+bool UMeshPaintMode::CanCopyMeshPaintTexture() const
+{
+	if (IsInSelectTool())
+	{
+		const TArray<UStaticMeshComponent*> StaticMeshComponents = GetSelectedComponents<UStaticMeshComponent>();
+		if (StaticMeshComponents.Num() == 1 && StaticMeshComponents[0]->GetMeshPaintTexture() != nullptr)
+		{
+			return true;
+		}
+	}
+	return false;
+}
+
+void UMeshPaintMode::Copy()
+{
+	FName PaletteName = Toolkit->GetCurrentPalette();
+	if (PaletteName == UMeshPaintMode::MeshPaintMode_VertexColor || PaletteName == UMeshPaintMode::MeshPaintMode_VertexWeights)
+	{
+		CopyInstanceVertexColors();
+	}
+	else if (PaletteName == UMeshPaintMode::MeshPaintMode_TextureColor)
+	{
+		CopyMeshPaintTexture();
+	}
+}
+
+bool UMeshPaintMode::CanCopy() const
+{
+	FName PaletteName = Toolkit->GetCurrentPalette();
+	if (PaletteName == UMeshPaintMode::MeshPaintMode_VertexColor || PaletteName == UMeshPaintMode::MeshPaintMode_VertexWeights)
+	{
+		return CanCopyInstanceVertexColors();
+	}
+	else if (PaletteName == UMeshPaintMode::MeshPaintMode_TextureColor)
+	{
+		return CanCopyMeshPaintTexture();
+	}
+	return false;
+}
+
+void UMeshPaintMode::PasteMeshPaintTexture()
+{
+	FScopedTransaction Transaction(LOCTEXT("LevelMeshPainter_TransactionPasteMeshPaintTexture", "Pasting Texture Colors"));
+	FImage const& Image = GEngine->GetEngineSubsystem<UMeshPaintingSubsystem>()->GetCopiedTexture();
+	const TArray<UStaticMeshComponent*> StaticMeshComponents = GetSelectedComponents<UStaticMeshComponent>();
+	for (UStaticMeshComponent* Component : StaticMeshComponents)
+	{
+		GEngine->GetEngineSubsystem<UMeshPaintingSubsystem>()->CreateComponentMeshPaintTexture(Component, Image);
+	}
+
+	// If this adds a mesh paint texture the selection/paint status may have changed.
+	const TArray<UMeshComponent*> CurrentMeshComponents = GetSelectedComponents<UMeshComponent>();
+	UpdateToolForSelection(CurrentMeshComponents);
+}
+
+bool UMeshPaintMode::CanPasteMeshPaintTexture() const
+{
+	if (!IsInSelectTool())
+	{
+		return false;
+	}
+	const TArray<UStaticMeshComponent*> StaticMeshComponents = GetSelectedComponents<UStaticMeshComponent>();
+	if (StaticMeshComponents.Num() == 0)
+	{
+		return false;
+	}
+	FImage const& Image = GEngine->GetEngineSubsystem<UMeshPaintingSubsystem>()->GetCopiedTexture();
+	if (Image.GetNumPixels() == 0)
+	{
+		return false;
+	}
+	return true;
+}
+
+void UMeshPaintMode::Paste()
+{
+	FName PaletteName = Toolkit->GetCurrentPalette();
+	if (PaletteName == UMeshPaintMode::MeshPaintMode_VertexColor || PaletteName == UMeshPaintMode::MeshPaintMode_VertexWeights)
+	{
+		PasteInstanceVertexColors();
+	}
+	else if (PaletteName == UMeshPaintMode::MeshPaintMode_TextureColor)
+	{
+		PasteMeshPaintTexture();
+	}
+}
+
+bool UMeshPaintMode::CanPaste() const
+{
+	FName PaletteName = Toolkit->GetCurrentPalette();
+	if (PaletteName == UMeshPaintMode::MeshPaintMode_VertexColor || PaletteName == UMeshPaintMode::MeshPaintMode_VertexWeights)
+	{
+		return CanPasteInstanceVertexColors();
+	}
+	else if (PaletteName == UMeshPaintMode::MeshPaintMode_TextureColor)
+	{
+		return CanPasteMeshPaintTexture();
+	}
+	return false;
+}
+
+bool UMeshPaintMode::CanImportVertexColorsFromFile() const
+{
+	FName PaletteName = Toolkit->GetCurrentPalette();
+	return PaletteName == UMeshPaintMode::MeshPaintMode_VertexColor || PaletteName == UMeshPaintMode::MeshPaintMode_VertexWeights;
+}
+
+void UMeshPaintMode::ImportVertexColorsFromMeshPaintTexture()
+{
+	FScopedTransaction Transaction(LOCTEXT("LevelMeshPainter_TransactionImportColors", "Importing Vertex Colors From Mesh Paint Textures"));
+	const TArray<UMeshComponent*> MeshComponents = GetSelectedComponents<UMeshComponent>();
+	for (UMeshComponent* Component : MeshComponents)
+	{
+		GEditor->GetEditorSubsystem<UMeshPaintModeSubsystem>()->ImportVertexColorsFromMeshPaintTexture(Component);
+	}
+}
+
+bool UMeshPaintMode::CanImportVertexColorsFromMeshPaintTexture() const
+{
+	const TArray<UMeshComponent*> MeshComponents = GetSelectedComponents<UMeshComponent>();
+	for (UMeshComponent* Component : MeshComponents)
+	{
+		if (Component->GetMeshPaintTexture() != nullptr)
+		{
+			return true;
+		}
+	}
+	return false;
+}
+
+void UMeshPaintMode::ImportMeshPaintTextureFromVertexColors()
+{
+	FScopedTransaction Transaction(LOCTEXT("LevelMeshPainter_TransactionImportTextures", "Importing Mesh Paint Textures From Vertex Colors"));
+	const TArray<UMeshComponent*> MeshComponents = GetSelectedComponents<UMeshComponent>();
+	for (UMeshComponent* Component : MeshComponents)
+	{
+		GEditor->GetEditorSubsystem<UMeshPaintModeSubsystem>()->ImportMeshPaintTextureFromVertexColors(Component);
+	}
+
+	// If this adds a mesh paint texture the selection/paint status may have changed.
+	const TArray<UMeshComponent*> CurrentMeshComponents = GetSelectedComponents<UMeshComponent>();
+	UpdateToolForSelection(CurrentMeshComponents);
+}
+
+bool UMeshPaintMode::CanImportMeshPaintTextureFromVertexColors() const
+{
+	if (IsInSelectTool())
+	{
+		const TArray<UMeshComponent*> MeshComponents = GetSelectedComponents<UMeshComponent>();
+		return MeshComponents.Num() > 0;
+	}
+	return false;
+}
+
+bool UMeshPaintMode::CanCycleMeshLODs() const
+{
+	FName PaletteName = Toolkit->GetCurrentPalette();
+	return PaletteName == UMeshPaintMode::MeshPaintMode_VertexColor || PaletteName == UMeshPaintMode::MeshPaintMode_VertexWeights;
+}
+
+#undef LOCTEXT_NAMESPACE
diff --git a/Engine/Plugins/MeshPainting/Source/MeshPaintEditorMode/Private/MeshPaintMode.h b/Engine/Plugins/MeshPainting/Source/MeshPaintEditorMode/Private/MeshPaintMode.h
--- a/Engine/Plugins/MeshPainting/Source/MeshPaintEditorMode/Private/MeshPaintMode.h
+++ b/Engine/Plugins/MeshPainting/Source/MeshPaintEditorMode/Private/MeshPaintMode.h
@@ -7,9 +7,12 @@
 #include "Tools/LegacyEdModeInterfaces.h"
 #include "MeshPaintMode.generated.h"
 
+class UMeshPaintingToolProperties;
 class UMeshVertexPaintingToolProperties;
-class UMeshColorPaintingToolProperties;
+class UMeshVertexColorPaintingToolProperties;
 class UMeshTexturePaintingToolProperties;
+class UMeshTextureColorPaintingToolProperties;
+class UMeshTextureAssetPaintingToolProperties;
 class UMeshPaintModeSettings;
 class IMeshPaintComponentAdapter;
 class UMeshComponent;
@@ -26,24 +29,29 @@ public:
 
 	/** Default constructor for UMeshPaintMode */
 	UMeshPaintMode();
+	static UMeshPaintingToolProperties* GetToolProperties();
 	static UMeshVertexPaintingToolProperties* GetVertexToolProperties();
-	static UMeshColorPaintingToolProperties* GetColorToolProperties();
-	static UMeshWeightPaintingToolProperties* GetWeightToolProperties();
+	static UMeshVertexColorPaintingToolProperties* GetVertexColorToolProperties();
+	static UMeshVertexWeightPaintingToolProperties* GetVertexWeightToolProperties();
 	static UMeshTexturePaintingToolProperties* GetTextureToolProperties();
+	static UMeshTextureColorPaintingToolProperties* GetTextureColorToolProperties();
+	static UMeshTextureAssetPaintingToolProperties* GetTextureAssetToolProperties();
 	static UMeshPaintMode* GetMeshPaintMode();
 	virtual void Enter() override;
 	virtual void Exit() override;
 	virtual void CreateToolkit() override;
 	virtual void Tick(FEditorViewportClient* ViewportClient, float DeltaTime) override;
 	virtual bool HandleClick(FEditorViewportClient* InViewportClient, HHitProxy* HitProxy, const FViewportClick& Click) override;
-	static FName MeshPaintMode_Color;
-	static FName MeshPaintMode_Texture;
-	static FName MeshPaintMode_Weights;
+	static FName MeshPaintMode_VertexColor;
+	static FName MeshPaintMode_VertexWeights;
+	static FName MeshPaintMode_TextureColor;
+	static FName MeshPaintMode_TextureAsset;
 	static FString VertexSelectToolName;
 	static FString TextureSelectToolName;
-	static FString ColorPaintToolName;
-	static FString WeightPaintToolName;
-	static FString TexturePaintToolName;
+	static FString VertexColorPaintToolName;
+	static FString VertexWeightPaintToolName;
+	static FString TextureColorPaintToolName;
+	static FString TextureAssetPaintToolName;
 
 	virtual TMap<FName, TArray<TSharedPtr<FUICommandInfo>>> GetModeCommands() const override;
 	/** Returns the instance of ComponentClass found in the current Editor selection */
@@ -67,51 +75,79 @@ protected:
 	virtual void ActivateDefaultTool() override;
 	virtual void UpdateOnPaletteChange(FName NewPalette);
 	// end UEdMode Interface
+	
 	void UpdateSelectedMeshes();
 
-	void CheckSelectionForTexturePaintCompat(const TArray<UMeshComponent*>& CurrentMeshComponents);
-	
+	void UpdateToolForSelection(const TArray<UMeshComponent*>& CurrentMeshComponents);
 
 	void UpdateOnMaterialChange(bool bInvalidateHitProxies);
 	void OnObjectsReplaced(const TMap<UObject*, UObject*>& InOldToNewInstanceMap);
 	void OnResetViewMode();
-	void FillWithVertexColor();
-	void PropagateVertexColorsToAsset();
-	bool CanPropagateVertexColors() const;
-	void ImportVertexColors();
-	void SavePaintedAssets();
-	bool CanSaveMeshPackages() const;
-	bool CanRemoveInstanceColors() const;
-	bool CanPasteInstanceVertexColors() const;
-	bool CanCopyInstanceVertexColors() const;
+	void OnVertexPaintFinished();
+
+	void UpdateCachedVertexDataSize();
+
+	bool IsInSelectTool() const;
+	bool IsInPaintTool() const;
+
+	// Start command bindings
+	void SwapColors();
+	bool CanSwapColors() const;
+	void FillVertexColors();
+	bool CanFillVertexColors() const;
+	void FillTexture();
+	bool CanFillTexture() const;
+	void ApplyVertexColorsToAsset();
+	bool CanApplyVertexColorsToAsset() const;
+	void CommitTextureColorsToAsset();
+	bool CanCommitTextureColorsToAsset() const;
+	void PropagateVertexColorsToLODs();
 	bool CanPropagateVertexColorsToLODs() const;
-	/** Copy and pasting functionality from and to a StaticMeshComponent (currently only supports per instance)*/
-	void CopyVertexColors();
-	void PasteVertexColors();
+	void SaveVertexColorsToAssets();
+	bool CanSaveVertexColorsToAssets() const;
+	void SaveTexturePackages();
+	bool CanSaveTexturePackages() const;
+	void AddMeshPaintTextures();
+	bool CanAddMeshPaintTextures() const;
+	void RemoveInstanceVertexColors();
+	bool CanRemoveInstanceVertexColors() const;
+	void RemoveMeshPaintTexture();
+	bool CanRemoveMeshPaintTextures() const;
+	void CopyInstanceVertexColors();
+	bool CanCopyInstanceVertexColors() const;
+	void CopyMeshPaintTexture();
+	bool CanCopyMeshPaintTexture() const;
+	void Copy();
+	bool CanCopy() const;
+	void PasteInstanceVertexColors();
+	bool CanPasteInstanceVertexColors() const;
+	void PasteMeshPaintTexture();
+	bool CanPasteMeshPaintTexture() const;
+	void Paste();
+	bool CanPaste() const;
+	void ImportVertexColorsFromFile();
+	bool CanImportVertexColorsFromFile() const;
+	void ImportVertexColorsFromMeshPaintTexture();
+	bool CanImportVertexColorsFromMeshPaintTexture() const;
+	void ImportMeshPaintTextureFromVertexColors();
+	bool CanImportMeshPaintTextureFromVertexColors() const;
 	void FixVertexColors();
-	bool DoesRequireVertexColorsFixup() const;
-	void RemoveVertexColors();
-	void PropagateVertexColorsToLODs();
-	void UpdateCachedVertexDataSize();
+	bool CanFixVertexColors() const;
 	void CycleMeshLODs(int32 Direction);
+	bool CanCycleMeshLODs() const;
 	void CycleTextures(int32 Direction);
 	bool CanCycleTextures() const;
-	void CommitAllPaintedTextures();
-	int32 GetNumberOfPendingPaintChanges();
-	void OnVertexPaintFinished();
-	void FillWithTextureColor();
+	// End command bindings
 
 protected:
 	UPROPERTY(Transient)
 	TObjectPtr<UMeshPaintModeSettings> ModeSettings;
 
-
 	// End vertex paint state
 	FGetSelectedMeshComponents MeshComponentDelegate;
 	uint32 CachedVertexDataSize;
 	bool bRecacheVertexDataSize;
 
 	FDelegateHandle PaletteChangedHandle;
-
 };
 
diff --git a/Engine/Plugins/MeshPainting/Source/MeshPaintEditorMode/Private/MeshPaintModeCommands.cpp b/Engine/Plugins/MeshPainting/Source/MeshPaintEditorMode/Private/MeshPaintModeCommands.cpp
--- a/Engine/Plugins/MeshPainting/Source/MeshPaintEditorMode/Private/MeshPaintModeCommands.cpp
+++ b/Engine/Plugins/MeshPainting/Source/MeshPaintEditorMode/Private/MeshPaintModeCommands.cpp
@@ -5,17 +5,12 @@
 
 #define LOCTEXT_NAMESPACE "MeshPaintEditorModeCommands"
 
-
-
 void FMeshPaintingToolActionCommands::GetToolDefaultObjectList(TArray<UInteractiveTool*>& ToolCDOs)
 {
 	ToolCDOs.Add(GetMutableDefault<USingleSelectionTool>());
 	ToolCDOs.Add(GetMutableDefault<UMeshColorPaintingTool>());
 }
 
-
-
-
 void FMeshPaintingToolActionCommands::RegisterAllToolActions()
 {
 	FMeshPaintingToolActionCommands::Register();
@@ -34,74 +29,55 @@ void FMeshPaintingToolActionCommands::UpdateToolCommandBinding(UInteractiveTool*
 	}
 }
 
-
-
-
 void FMeshPaintEditorModeCommands::RegisterCommands()
 {
-	TArray<TSharedPtr<FUICommandInfo>> ColorCommands;
-	TArray<TSharedPtr<FUICommandInfo>> WeightCommands;
-	TArray<TSharedPtr<FUICommandInfo>> VertexCommands;
-	TArray<TSharedPtr<FUICommandInfo>> TextureCommands;
-	 
-	UI_COMMAND(VertexSelect, "Select", "Select the mesh for vertex painting", EUserInterfaceActionType::ToggleButton, FInputChord());
-	ColorCommands.Add(VertexSelect);
-	WeightCommands.Add(VertexSelect);
-
-	UI_COMMAND(ColorPaint, "Paint", "Paint the mesh", EUserInterfaceActionType::ToggleButton, FInputChord());
-	ColorCommands.Add(ColorPaint);
-	UI_COMMAND(WeightPaint, "Paint", "Paint the mesh", EUserInterfaceActionType::ToggleButton, FInputChord());
-	WeightCommands.Add(WeightPaint);
-	UI_COMMAND(SwitchForeAndBackgroundColor, "Swap", "Switches the Foreground and Background Colors used for Vertex Painting", EUserInterfaceActionType::Button, FInputChord(EKeys::X));
-	VertexCommands.Add(SwitchForeAndBackgroundColor);
-	UI_COMMAND(Fill, "Fill", "Fills the selected Meshes with the Paint Color", EUserInterfaceActionType::Button, FInputChord());
-	VertexCommands.Add(Fill);
-	UI_COMMAND(Propagate, "Apply", "Propagates Instance Vertex Colors to the Source Meshes", EUserInterfaceActionType::Button, FInputChord());
-	VertexCommands.Add(Propagate);
- 	UI_COMMAND(Import, "Import", "Imports Vertex Colors from a TGA Texture File to the Selected Meshes", EUserInterfaceActionType::Button, FInputChord());
- 	VertexCommands.Add(Import);
- 	UI_COMMAND(Save, "Save", "Saves the Source Meshes for the selected Mesh Components", EUserInterfaceActionType::Button, FInputChord());
-	VertexCommands.Add(Save);
-	UI_COMMAND(Copy, "Copy", "Copies Vertex Colors from the selected Mesh Components", EUserInterfaceActionType::Button, FInputChord());
- 	VertexCommands.Add(Copy);
-	UI_COMMAND(Paste, "Paste", "Tried to Paste Vertex Colors on the selected Mesh Components", EUserInterfaceActionType::Button, FInputChord());
- 	VertexCommands.Add(Paste);
-	UI_COMMAND(Remove, "Remove", "Removes Vertex Colors from the selected Mesh Components", EUserInterfaceActionType::Button, FInputChord());
-	VertexCommands.Add(Remove);
-	UI_COMMAND(Fix, "Fix", "If necessary fixes Vertex Colors applied to the selected Mesh Components", EUserInterfaceActionType::Button, FInputChord());
-	VertexCommands.Add(Fix);
-	UI_COMMAND(PropagateVertexColorsToLODs, "All LODs", "Applied the Vertex Colors from LOD0 to all LOD levels", EUserInterfaceActionType::Button, FInputChord());
-	VertexCommands.Add(PropagateVertexColorsToLODs);
-
-	UI_COMMAND(CycleToPreviousLOD, "Previous LOD", "Cycles to the previous possible Mesh LOD to Paint on", EUserInterfaceActionType::Button, FInputChord(EKeys::B));
-	UI_COMMAND(CycleToNextLOD, "Next LOD", "Cycles to the next possible Mesh LOD to Paint on", EUserInterfaceActionType::Button, FInputChord(EKeys::N));
-
-	ColorCommands.Append(VertexCommands);
-	WeightCommands.Append(VertexCommands);
-	Commands.Add(UMeshPaintMode::MeshPaintMode_Color, ColorCommands);
-	Commands.Add(UMeshPaintMode::MeshPaintMode_Weights, WeightCommands);
-
-	UI_COMMAND(TextureSelect, "Select", "Select the mesh for texture painting", EUserInterfaceActionType::ToggleButton, FInputChord());
-	TextureCommands.Add(TextureSelect);
-	UI_COMMAND(TexturePaint, "Paint", "Paint the mesh", EUserInterfaceActionType::ToggleButton, FInputChord());
-	TextureCommands.Add(TexturePaint);
-	UI_COMMAND(TextureFill, "Fill", "Fills the selected Mesh's texture with the Paint Color", EUserInterfaceActionType::Button, FInputChord());
-	TextureCommands.Add(TextureFill);
-	
-	UI_COMMAND(PreviousTexture, "Previous Texture", "Cycle To Previous Texture", EUserInterfaceActionType::Button, FInputChord(EKeys::Comma));
-	UI_COMMAND(NextTexture, "Next Texture", "Cycle To Next Texture", EUserInterfaceActionType::Button, FInputChord(EKeys::Period));
-
-	UI_COMMAND(PropagateTexturePaint, "Commit", "Commits Texture Painting Changes", EUserInterfaceActionType::Button, FInputChord(EKeys::C, EModifierKey::Control | EModifierKey::Shift));
-	TextureCommands.Add(PropagateTexturePaint);
-	UI_COMMAND(SaveTexturePaint, "Save", "Saves the Modified Textures for the selected Mesh Components", EUserInterfaceActionType::Button, FInputChord());
-	TextureCommands.Add(SaveTexturePaint);
-
-	Commands.Add(UMeshPaintMode::MeshPaintMode_Texture, TextureCommands);
-
-
-
-
+	UI_COMMAND(SelectVertex, "Select", "Select the mesh for vertex painting", EUserInterfaceActionType::ToggleButton, FInputChord());
+	UI_COMMAND(SelectTexture, "Select", "Select the mesh for texture painting", EUserInterfaceActionType::ToggleButton, FInputChord());
+	UI_COMMAND(PaintVertexColor, "Paint", "Paint the mesh vertex colors", EUserInterfaceActionType::ToggleButton, FInputChord());
+	UI_COMMAND(PaintVertexWeight, "Paint", "Paint the mesh vertex weights", EUserInterfaceActionType::ToggleButton, FInputChord());
+	UI_COMMAND(PaintTextureColor, "Paint", "Paint the mesh texture colors", EUserInterfaceActionType::ToggleButton, FInputChord());
+	UI_COMMAND(PaintTextureAsset, "Paint", "Paint texture assets used by the mesh material", EUserInterfaceActionType::ToggleButton, FInputChord());
+	UI_COMMAND(SwapColor, "Swap", "Switches the foreground and background colors used for painting", EUserInterfaceActionType::Button, FInputChord(EKeys::X));
+	UI_COMMAND(FillVertex, "Fill", "Fills the selected meshes with the paint color", EUserInterfaceActionType::Button, FInputChord());
+	UI_COMMAND(FillTexture, "Fill", "Fills the selected textures with the paint color", EUserInterfaceActionType::Button, FInputChord());
+	UI_COMMAND(Apply, "Apply", "Propagates instance vertex colors to the source meshes", EUserInterfaceActionType::Button, FInputChord());
+	UI_COMMAND(Commit, "Commit", "Commits texture edits to the source texture assets", EUserInterfaceActionType::Button, FInputChord(EKeys::C, EModifierKey::Control | EModifierKey::Shift));
+	UI_COMMAND(Propagate, "All LODs", "Applies the vertex colors from LOD0 to all LOD levels", EUserInterfaceActionType::Button, FInputChord());
+	UI_COMMAND(SaveVertex, "Save", "Saves the source meshes for the selected mesh components", EUserInterfaceActionType::Button, FInputChord());
+	UI_COMMAND(SaveTexture, "Save", "Saves the modified textures for the selected mesh components", EUserInterfaceActionType::Button, FInputChord());
+	UI_COMMAND(Add, "Add", "Adds mesh paint textures to the selected mesh components to enable painting", EUserInterfaceActionType::Button, FInputChord());
+	UI_COMMAND(RemoveVertex, "Remove", "Removes any vertex colors from the selected mesh components", EUserInterfaceActionType::Button, FInputChord());
+	UI_COMMAND(RemoveTexture, "Remove", "Removes any mesh paint textures from the selected mesh components", EUserInterfaceActionType::Button, FInputChord());
+	UI_COMMAND(Copy, "Copy", "Copies colors from the selected mesh components", EUserInterfaceActionType::Button, FInputChord());
+	UI_COMMAND(Paste, "Paste", "Pastes colors on the selected mesh components", EUserInterfaceActionType::Button, FInputChord());
+	UI_COMMAND(Import, "Import", "Imports vertex colors from a TGA texture file to the selected meshes", EUserInterfaceActionType::Button, FInputChord());
+	UI_COMMAND(GetVertexColors, "Vertex", "Imports texture colors from vertex colors on the selected mesh components", EUserInterfaceActionType::Button, FInputChord());
+	UI_COMMAND(GetTextureColors, "Texture", "Imports vertex colors from texture colors on the selected mesh components", EUserInterfaceActionType::Button, FInputChord());
+	UI_COMMAND(Fix, "Fix", "Applies any required color data fixes to the selected mesh components", EUserInterfaceActionType::Button, FInputChord());
+	UI_COMMAND(PreviousLOD, "Previous LOD", "Cycles to the previous possible mesh LOD to paint on", EUserInterfaceActionType::Button, FInputChord(EKeys::B));
+	UI_COMMAND(NextLOD, "Next LOD", "Cycles to the next possible mesh LOD to paint on", EUserInterfaceActionType::Button, FInputChord(EKeys::N));
+	UI_COMMAND(PreviousTexture, "Previous Texture", "Cycle To previous texture", EUserInterfaceActionType::Button, FInputChord(EKeys::Comma));
+	UI_COMMAND(NextTexture, "Next Texture", "Cycle To next texture", EUserInterfaceActionType::Button, FInputChord(EKeys::Period));
+
+	TArray<TSharedPtr<FUICommandInfo>> VertexColorCommands = {
+		SelectVertex, PaintVertexColor, SwapColor, FillVertex, Apply, Propagate, SaveVertex, RemoveVertex, Copy, Paste, Import, GetTextureColors, Fix
+	};
+	Commands.Add(UMeshPaintMode::MeshPaintMode_VertexColor, VertexColorCommands);
+
+	TArray<TSharedPtr<FUICommandInfo>> VertexWeightCommands = {
+		SelectVertex, PaintVertexWeight, SwapColor, FillVertex, Apply, Propagate, SaveVertex, RemoveVertex, Copy, Paste, Import, Fix
+	};
+	Commands.Add(UMeshPaintMode::MeshPaintMode_VertexWeights, VertexWeightCommands);
+
+	TArray<TSharedPtr<FUICommandInfo>> TextureColorCommands = {
+		SelectTexture, PaintTextureColor, SwapColor, FillTexture, Commit, SaveTexture, Add, RemoveTexture, Copy, Paste, GetVertexColors
+	};
+	Commands.Add(UMeshPaintMode::MeshPaintMode_TextureColor, TextureColorCommands);
+
+	TArray<TSharedPtr<FUICommandInfo>> TextureAssetCommands = {
+		SelectTexture, PaintTextureAsset, SwapColor, FillTexture, Commit, SaveTexture
+	};
+	Commands.Add(UMeshPaintMode::MeshPaintMode_TextureAsset, TextureAssetCommands);
 }
 
 #undef LOCTEXT_NAMESPACE
-
diff --git a/Engine/Plugins/MeshPainting/Source/MeshPaintEditorMode/Private/MeshPaintModeHelpers.cpp b/Engine/Plugins/MeshPainting/Source/MeshPaintEditorMode/Private/MeshPaintModeHelpers.cpp
--- a/Engine/Plugins/MeshPainting/Source/MeshPaintEditorMode/Private/MeshPaintModeHelpers.cpp
+++ b/Engine/Plugins/MeshPainting/Source/MeshPaintEditorMode/Private/MeshPaintModeHelpers.cpp
@@ -41,6 +41,12 @@
 #include "InterchangeGenericAssetsPipeline.h"
 #include "InterchangePythonPipelineBase.h"
 
+#include "Math/Color.h"
+#include "StaticMeshLODResourcesToDynamicMesh.h"
+#include "Sampling/MeshMapBaker.h"
+#include "Sampling/MeshPropertyMapEvaluator.h"
+#include "VT/MeshPaintVirtualTexture.h"
+
 
 void UMeshPaintModeSubsystem::SetViewportColorMode(EMeshPaintActiveMode ActiveMode, EMeshPaintDataColorViewMode ColorViewMode, FEditorViewportClient* ViewportClient)
 {
@@ -113,7 +119,7 @@ void UMeshPaintModeSubsystem::SetViewportColorMode(EMeshPaintActiveMode ActiveMo
 				if (ActiveMode == EMeshPaintActiveMode::Texture)
 				{
 					UMeshPaintingSubsystem* MeshPaintingSubsystem = GEngine->GetEngineSubsystem<UMeshPaintingSubsystem>();
-					UMeshTexturePaintingToolProperties* Settings = UMeshPaintMode::GetTextureToolProperties();
+					UMeshTextureAssetPaintingToolProperties* Settings = UMeshPaintMode::GetTextureAssetToolProperties();
 					if (MeshPaintingSubsystem && MeshPaintingSubsystem->OverridePaintTexture.IsValid())
 					{
 						SelectedTexture = MeshPaintingSubsystem->OverridePaintTexture.Get();
@@ -271,6 +277,89 @@ void UMeshPaintModeSubsystem::ImportVertexColorsFromTexture(UMeshComponent* Mesh
 	}
 }
 
+void UMeshPaintModeSubsystem::ImportVertexColorsFromMeshPaintTexture(UMeshComponent* MeshComponent)
+{
+	if (UTexture2D* Texture = Cast<UTexture2D>(MeshComponent->GetMeshPaintTexture()))
+	{
+		UImportVertexColorOptions* Options = NewObject<UImportVertexColorOptions>();
+ 		Options->UVIndex = MeshComponent->GetMeshPaintTextureCoordinateIndex();
+
+		if (UStaticMeshComponent* StaticMeshComponent = Cast<UStaticMeshComponent>(MeshComponent))
+		{
+			ImportVertexColorsToStaticMeshComponent(StaticMeshComponent, Options, Texture);
+		}
+	}
+}
+
+void UMeshPaintModeSubsystem::ImportMeshPaintTextureFromVertexColors(UMeshComponent* MeshComponent)
+{
+	UStaticMeshComponent* StaticMeshComponent = Cast<UStaticMeshComponent>(MeshComponent);
+	if (StaticMeshComponent == nullptr)
+	{
+		return;
+	}
+
+	UStaticMesh* StaticMesh = StaticMeshComponent->GetStaticMesh();
+	if (StaticMesh == nullptr)
+	{
+		return;
+	}
+
+	const int32 LodIndex = 0;
+
+	FStaticMeshComponentLODInfo* InstanceMeshLODInfo = StaticMeshComponent->LODData.IsValidIndex(LodIndex) ? &StaticMeshComponent->LODData[LodIndex] : nullptr;
+	const bool bHasPerInstanceVertexColors = InstanceMeshLODInfo != nullptr && InstanceMeshLODInfo->OverrideVertexColors != nullptr;
+
+	UE::Geometry::FStaticMeshLODResourcesToDynamicMesh::ConversionOptions ConversionOptions;
+	ConversionOptions.bWantTangents = false;
+	ConversionOptions.bWantMaterialIDs = false;
+
+	UE::Geometry::FDynamicMesh3 DynamicMesh;
+	UE::Geometry::FStaticMeshLODResourcesToDynamicMesh Converter;
+	Converter.Convert(
+		&StaticMesh->GetRenderData()->LODResources[LodIndex],
+		ConversionOptions,
+		DynamicMesh,
+		bHasPerInstanceVertexColors,
+		[InstanceMeshLODInfo](int32 Index)
+		{
+			return InstanceMeshLODInfo->OverrideVertexColors->VertexColor(Index);
+		});
+
+	const int32 TextureSize = MeshPaintVirtualTexture::GetDefaultTextureSize(DynamicMesh.VertexCount());
+
+	const UE::Geometry::FDynamicMeshAABBTree3 DetailSpatial(&DynamicMesh);
+	UE::Geometry::FMeshBakerDynamicMeshSampler DetailSampler(&DynamicMesh, &DetailSpatial);
+
+	TSharedPtr<UE::Geometry::FMeshPropertyMapEvaluator, ESPMode::ThreadSafe> PropertyEval = MakeShared<UE::Geometry::FMeshPropertyMapEvaluator, ESPMode::ThreadSafe>();
+	PropertyEval->Property = UE::Geometry::EMeshPropertyMapType::VertexColor;
+
+	UE::Geometry::FMeshMapBaker Baker;
+	Baker.SetTargetMesh(&DynamicMesh);
+	Baker.SetDetailSampler(&DetailSampler);
+	Baker.AddEvaluator(PropertyEval);
+	Baker.SetTargetMeshUVLayer(StaticMeshComponent->GetMeshPaintTextureCoordinateIndex());
+	Baker.SetDimensions(UE::Geometry::FImageDimensions(TextureSize, TextureSize));
+	Baker.SetProjectionDistance(3.0f);
+	Baker.SetSamplesPerPixel(1);
+	Baker.SetFilter(UE::Geometry::FMeshMapBaker::EBakeFilterType::BSpline);
+	Baker.SetGutterEnabled(true);
+	Baker.SetGutterSize(4);
+	Baker.Bake();
+
+	FImageView ResultImage((FLinearColor*)Baker.GetBakeResults(0)[0]->GetImageBuffer().GetData(), TextureSize, TextureSize);
+	FImage ConvertedImage;
+	ResultImage.CopyTo(ConvertedImage, ERawImageFormat::BGRA8, EGammaSpace::sRGB);
+
+	UMeshPaintVirtualTexture* NewTexture = NewObject<UMeshPaintVirtualTexture>(StaticMeshComponent->GetOutermost());
+	NewTexture->Source.Init(ConvertedImage);
+	NewTexture->OwningComponent = MakeWeakObjectPtr(StaticMeshComponent);
+	NewTexture->UpdateResource();
+
+	StaticMeshComponent->Modify();
+	StaticMeshComponent->SetMeshPaintTexture(NewTexture);
+}
+
 
 void UMeshPaintModeSubsystem::ImportVertexColorsToSkeletalMesh(USkeletalMesh* SkeletalMesh, const UImportVertexColorOptions* Options, UTexture2D* Texture)
 {
@@ -812,10 +901,9 @@ void UMeshPaintModeSubsystem::RemovePerLODColors(const TArray<UMeshComponent*>&
 	}
 }
 
-void UMeshPaintModeSubsystem::SwapVertexColors()
+void UMeshPaintModeSubsystem::SwapColors()
 {
-	UMeshVertexPaintingToolProperties* Settings = UMeshPaintMode::GetVertexToolProperties();
-	if (Settings)
+	if (UMeshPaintingToolProperties* Settings = UMeshPaintMode::GetToolProperties())
 	{
 		Settings->Modify();
 
@@ -824,37 +912,3 @@ void UMeshPaintModeSubsystem::SwapVertexColors()
 		Settings->EraseColor = TempPaintColor;
 	}
 }
-
-
-
-void UMeshPaintModeSubsystem::SaveModifiedTextures()
-{
-	UMeshTexturePaintingToolProperties* Settings = UMeshPaintMode::GetTextureToolProperties();
-	if (Settings)
-	{
-		UTexture2D* SelectedTexture = Settings->PaintTexture;
-
-		if (nullptr != SelectedTexture)
-		{
-			TArray<UObject*> TexturesToSaveArray;
-			TexturesToSaveArray.Add(SelectedTexture);
-			UPackageTools::SavePackagesForObjects(TexturesToSaveArray);
-		}
-	}
-}
-
-bool UMeshPaintModeSubsystem::CanSaveModifiedTextures()
-{
-	/** Check whether or not the current selected paint texture requires saving */
-	bool bRequiresSaving = false;
-	UMeshTexturePaintingToolProperties* Settings = UMeshPaintMode::GetTextureToolProperties();
-	if (Settings)
-	{
-		const UTexture2D* SelectedTexture = Settings->PaintTexture;
-		if (nullptr != SelectedTexture)
-		{
-			bRequiresSaving = SelectedTexture->GetOutermost()->IsDirty();
-		}
-	}
-	return bRequiresSaving;
-}
diff --git a/Engine/Plugins/MeshPainting/Source/MeshPaintEditorMode/Private/MeshPaintModeToolkit.cpp b/Engine/Plugins/MeshPainting/Source/MeshPaintEditorMode/Private/MeshPaintModeToolkit.cpp
--- a/Engine/Plugins/MeshPainting/Source/MeshPaintEditorMode/Private/MeshPaintModeToolkit.cpp
+++ b/Engine/Plugins/MeshPainting/Source/MeshPaintEditorMode/Private/MeshPaintModeToolkit.cpp
@@ -24,7 +24,6 @@ FName FMeshPaintModeToolkit::GetToolkitFName() const
 	return FName( "MeshPaintMode" );
 }
 
-
 FText FMeshPaintModeToolkit::GetBaseToolkitName() const
 {
 	return LOCTEXT( "ToolkitName", "Mesh Paint Mode" );
@@ -32,25 +31,30 @@ FText FMeshPaintModeToolkit::GetBaseToolkitName() const
 
 void FMeshPaintModeToolkit::GetToolPaletteNames(TArray<FName>& PaletteNames) const
 {
-	PaletteNames.Add(UMeshPaintMode::MeshPaintMode_Color);
-	PaletteNames.Add(UMeshPaintMode::MeshPaintMode_Weights);
-	PaletteNames.Add(UMeshPaintMode::MeshPaintMode_Texture);
+	PaletteNames.Add(UMeshPaintMode::MeshPaintMode_VertexColor);
+	PaletteNames.Add(UMeshPaintMode::MeshPaintMode_VertexWeights);
+	PaletteNames.Add(UMeshPaintMode::MeshPaintMode_TextureColor);
+	PaletteNames.Add(UMeshPaintMode::MeshPaintMode_TextureAsset);
 }
 
 
 FText FMeshPaintModeToolkit::GetToolPaletteDisplayName(FName Palette) const
 {
-	if (Palette == UMeshPaintMode::MeshPaintMode_Color)
+	if (Palette == UMeshPaintMode::MeshPaintMode_VertexColor)
+	{
+		return LOCTEXT("MeshPaintMode_VertexColor", "Vertex\nColor");
+	}
+	if (Palette == UMeshPaintMode::MeshPaintMode_VertexWeights)
 	{
-		return LOCTEXT("MeshPaintMode_Color", "Colors");
+		return LOCTEXT("MeshPaintMode_VertexWeights", "Vertex\nWeights");
 	}
-	if (Palette == UMeshPaintMode::MeshPaintMode_Weights)
+	if (Palette == UMeshPaintMode::MeshPaintMode_TextureColor)
 	{
-		return LOCTEXT("MeshPaintMode_Weights", "Weights");
+		return LOCTEXT("MeshPaintMode_TextureColor", "Texture\nColor");
 	}
-	if (Palette == UMeshPaintMode::MeshPaintMode_Texture)
+	if (Palette == UMeshPaintMode::MeshPaintMode_TextureAsset)
 	{
-		return LOCTEXT("MeshPaintMode_Texture", "Textures");
+		return LOCTEXT("MeshPaintMode_TextureAsset", "Textures");
 	}
 	return FText();
 }
diff --git a/Engine/Plugins/MeshPainting/Source/MeshPaintEditorMode/Public/MeshPaintModeCommands.h b/Engine/Plugins/MeshPainting/Source/MeshPaintEditorMode/Public/MeshPaintModeCommands.h
--- a/Engine/Plugins/MeshPainting/Source/MeshPaintEditorMode/Public/MeshPaintModeCommands.h
+++ b/Engine/Plugins/MeshPainting/Source/MeshPaintEditorMode/Public/MeshPaintModeCommands.h
@@ -30,34 +30,37 @@ public:
 	}
 
 public:
-	TSharedPtr<FUICommandInfo> NextTexture;
-	TSharedPtr<FUICommandInfo> PreviousTexture;
-	TSharedPtr<FUICommandInfo> CommitTexturePainting;
-
+	TSharedPtr<FUICommandInfo> SelectVertex;
+	TSharedPtr<FUICommandInfo> SelectTexture;
+
+	TSharedPtr<FUICommandInfo> PaintVertexColor;
+	TSharedPtr<FUICommandInfo> PaintVertexWeight;
+	TSharedPtr<FUICommandInfo> PaintTextureColor;
+	TSharedPtr<FUICommandInfo> PaintTextureAsset;
+
+	TSharedPtr<FUICommandInfo> SwapColor;
+	TSharedPtr<FUICommandInfo> FillVertex;
+	TSharedPtr<FUICommandInfo> FillTexture;
+	TSharedPtr<FUICommandInfo> Apply;
+	TSharedPtr<FUICommandInfo> Commit;
+	TSharedPtr<FUICommandInfo> Propagate;
+	TSharedPtr<FUICommandInfo> SaveVertex;
+	TSharedPtr<FUICommandInfo> SaveTexture;
+	TSharedPtr<FUICommandInfo> Add;
+	TSharedPtr<FUICommandInfo> RemoveVertex;
+	TSharedPtr<FUICommandInfo> RemoveTexture;
 	TSharedPtr<FUICommandInfo> Copy;
 	TSharedPtr<FUICommandInfo> Paste;
-	TSharedPtr<FUICommandInfo> Remove;
-	TSharedPtr<FUICommandInfo> Fix;
-	TSharedPtr<FUICommandInfo> Fill;
-	TSharedPtr<FUICommandInfo> Propagate;
 	TSharedPtr<FUICommandInfo> Import;
-	TSharedPtr<FUICommandInfo> Save;
-
-	TSharedPtr<FUICommandInfo> SwitchForeAndBackgroundColor;
-	TSharedPtr<FUICommandInfo> CycleToNextLOD;
-	TSharedPtr<FUICommandInfo> CycleToPreviousLOD;
-
-	TSharedPtr<FUICommandInfo> PropagateTexturePaint;
-	TSharedPtr<FUICommandInfo> SaveTexturePaint;
+	TSharedPtr<FUICommandInfo> GetTextureColors;
+	TSharedPtr<FUICommandInfo> GetVertexColors;
+	TSharedPtr<FUICommandInfo> Fix;
 
-	TSharedPtr<FUICommandInfo> PropagateVertexColorsToLODs;
+	TSharedPtr<FUICommandInfo> PreviousLOD;
+	TSharedPtr<FUICommandInfo> NextLOD;
+	TSharedPtr<FUICommandInfo> PreviousTexture;
+	TSharedPtr<FUICommandInfo> NextTexture;
 
-	TSharedPtr<FUICommandInfo> ColorPaint;
-	TSharedPtr<FUICommandInfo> WeightPaint;
-	TSharedPtr<FUICommandInfo> TexturePaint;
-	TSharedPtr<FUICommandInfo> VertexSelect;
-	TSharedPtr<FUICommandInfo> TextureSelect;
-	TSharedPtr<FUICommandInfo> TextureFill;
 	TMap<FName, TArray<TSharedPtr<FUICommandInfo>>> Commands;
 };
 
diff --git a/Engine/Plugins/MeshPainting/Source/MeshPaintEditorMode/Public/MeshPaintModeHelpers.h b/Engine/Plugins/MeshPainting/Source/MeshPaintEditorMode/Public/MeshPaintModeHelpers.h
--- a/Engine/Plugins/MeshPainting/Source/MeshPaintEditorMode/Public/MeshPaintModeHelpers.h
+++ b/Engine/Plugins/MeshPainting/Source/MeshPaintEditorMode/Public/MeshPaintModeHelpers.h
@@ -27,8 +27,9 @@ struct FPerComponentVertexColorData;
 UENUM()
 enum class EMeshPaintActiveMode : uint8
 {
-	Color UMETA(DisplayName = "Color"),
-	Weights UMETA(DisplayName = "Weights"),
+	VertexColor UMETA(DisplayName = "VertexColor"),
+	VertexWeights UMETA(DisplayName = "VertexWeights"),
+	TextureColor UMETA(DisplayName = "TextureColor"),
 	Texture UMETA(DisplayName = "Texture"),
 };
 
@@ -50,6 +51,12 @@ public:
 	/** Imports vertex colors from a Texture to the specified Skeletal Mesh according to user-set options */
 	void ImportVertexColorsToSkeletalMesh(USkeletalMesh* SkeletalMesh, const UImportVertexColorOptions* Options, UTexture2D* Texture);
 
+	/** Helper function to import Vertex Colors from a the MeshPaintTexture on the mesh component */
+	void ImportVertexColorsFromMeshPaintTexture(UMeshComponent* MeshComponent);
+
+	/** Helper function to import the MeshPaintTexture on the mesh component from the vertex colors */
+	void ImportMeshPaintTextureFromVertexColors(UMeshComponent* MeshComponent);
+
 	struct FPaintRay
 	{
 		FVector CameraLocation;
@@ -75,9 +82,7 @@ public:
 	bool CanPasteInstanceVertexColors(const TArray<UStaticMeshComponent*>& StaticMeshComponents, const TArray<FPerComponentVertexColorData>& CopiedColorsByComponent);
 	void RemovePerLODColors(const TArray<UMeshComponent*>& PaintableComponents);
 
-	void SwapVertexColors();
-	void SaveModifiedTextures();
-	bool CanSaveModifiedTextures();
+	void SwapColors();
 };
 
 #if UE_ENABLE_INCLUDE_ORDER_DEPRECATED_IN_5_2
diff --git a/Engine/Plugins/MeshPainting/Source/MeshPaintingToolset/MeshPaintingToolset.Build.cs b/Engine/Plugins/MeshPainting/Source/MeshPaintingToolset/MeshPaintingToolset.Build.cs
--- a/Engine/Plugins/MeshPainting/Source/MeshPaintingToolset/MeshPaintingToolset.Build.cs
+++ b/Engine/Plugins/MeshPainting/Source/MeshPaintingToolset/MeshPaintingToolset.Build.cs
@@ -21,7 +21,8 @@ namespace UnrealBuildTool.Rules
 					"UnrealEd",
 					"InterchangeEngine",
 					"InterchangePipelines",
-					"Chaos"
+					"Chaos",
+					"ImageCore"
 				}
                 );
 
diff --git a/Engine/Plugins/MeshPainting/Source/MeshPaintingToolset/Private/IMeshPaintComponentAdapter.cpp b/Engine/Plugins/MeshPainting/Source/MeshPaintingToolset/Private/IMeshPaintComponentAdapter.cpp
--- a/Engine/Plugins/MeshPainting/Source/MeshPaintingToolset/Private/IMeshPaintComponentAdapter.cpp
+++ b/Engine/Plugins/MeshPainting/Source/MeshPaintingToolset/Private/IMeshPaintComponentAdapter.cpp
@@ -110,9 +110,10 @@ void IMeshPaintComponentAdapter::DefaultQueryPaintableTextures(int32 MaterialInd
 				// If material samples the primitive mesh paint texture, then add it here.
 				if (UMaterialExpressionMeshPaintTextureObject* MeshPaintTextureExpression = Cast<UMaterialExpressionMeshPaintTextureObject>(Expression))
 				{
-					if (UTexture* MesPaintTexture = MeshComponent->GetMeshPaintTexture())
+					if (UTexture* MeshPaintTexture = MeshComponent->GetMeshPaintTexture())
 					{
-						InOutTextureList.AddUnique(FPaintableTexture{ MesPaintTexture, 0 });
+						const int32 CoordinateIndex = MeshComponent->GetMeshPaintTextureCoordinateIndex();
+						InOutTextureList.AddUnique(FPaintableTexture(MeshPaintTexture, CoordinateIndex, true));
 					}
 				}
 			}
@@ -434,7 +435,8 @@ namespace UE::MeshPaintingToolset
 		}
 
 		// Check to see if the source texture is the special mesh paint texture on the component.
-		if (InMeshComponent->GetMeshPaintTexture() == SourceTexture)
+		// But always apply setting override to nullptr which can happen after the SourceTexture is cleared from the component.
+		if (InMeshComponent->GetMeshPaintTexture() == SourceTexture || !OverrideTexture)
 		{
 			InMeshComponent->SetMeshPaintTextureOverride(OverrideTexture);
 		}
diff --git a/Engine/Plugins/MeshPainting/Source/MeshPaintingToolset/Private/MeshPaintHelpers.cpp b/Engine/Plugins/MeshPainting/Source/MeshPaintingToolset/Private/MeshPaintHelpers.cpp
--- a/Engine/Plugins/MeshPainting/Source/MeshPaintingToolset/Private/MeshPaintHelpers.cpp
+++ b/Engine/Plugins/MeshPainting/Source/MeshPaintingToolset/Private/MeshPaintHelpers.cpp
@@ -15,6 +15,7 @@
 #include "Rendering/SkeletalMeshRenderData.h"
 #include "Rendering/SkeletalMeshModel.h"
 #include "MeshVertexPaintingTool.h"
+#include "VT/MeshPaintVirtualTexture.h"
 
 #include UE_INLINE_GENERATED_CPP_BY_NAME(MeshPaintHelpers)
 
@@ -74,6 +75,49 @@ void UMeshPaintingSubsystem::RemoveComponentInstanceVertexColors(UStaticMeshComp
 	}
 }
 
+void UMeshPaintingSubsystem::CreateComponentMeshPaintTexture(UStaticMeshComponent* StaticMeshComponent)
+{
+	if (StaticMeshComponent != nullptr && StaticMeshComponent->GetMeshPaintTexture() == nullptr)
+	{
+		StaticMeshComponent->Modify();
+
+		const uint32 TextureSize = MeshPaintVirtualTexture::GetDefaultTextureSize(StaticMeshComponent->GetStaticMesh()->GetNumVertices(0));
+		const uint32 TextureNumMips = FMath::FloorLog2(TextureSize) + 1;
+
+		UMeshPaintVirtualTexture* NewTexture = NewObject<UMeshPaintVirtualTexture>(StaticMeshComponent->GetOutermost());
+		NewTexture->Source.Init(TextureSize, TextureSize, 1, TextureNumMips, TSF_BGRA8);
+		NewTexture->OwningComponent = MakeWeakObjectPtr(StaticMeshComponent);
+		NewTexture->UpdateResource();
+
+		StaticMeshComponent->SetMeshPaintTexture(NewTexture);
+	}
+}
+
+void UMeshPaintingSubsystem::CreateComponentMeshPaintTexture(UStaticMeshComponent* StaticMeshComponent, FImageView const& InImage)
+{
+	if (StaticMeshComponent != nullptr)
+	{
+		StaticMeshComponent->Modify();
+
+		UMeshPaintVirtualTexture* NewTexture = NewObject<UMeshPaintVirtualTexture>(StaticMeshComponent->GetOutermost());
+		NewTexture->Source.Init(InImage);
+		NewTexture->OwningComponent = MakeWeakObjectPtr(StaticMeshComponent);
+		NewTexture->UpdateResource();
+
+		StaticMeshComponent->SetMeshPaintTexture(NewTexture);
+	}
+}
+
+void UMeshPaintingSubsystem::RemoveComponentMeshPaintTexture(UStaticMeshComponent* StaticMeshComponent)
+{
+	if (StaticMeshComponent != nullptr && StaticMeshComponent->GetMeshPaintTexture() != nullptr)
+	{
+		// Mark the mesh component as modified
+		StaticMeshComponent->Modify();
+
+		StaticMeshComponent->SetMeshPaintTexture(nullptr);
+	}
+}
 
 bool UMeshPaintingSubsystem::PropagateColorsToRawMesh(UStaticMesh* StaticMesh, int32 LODIndex, FStaticMeshComponentLODInfo& ComponentLODInfo)
 {
@@ -1120,6 +1164,12 @@ FColor UMeshPaintingSubsystem::PickVertexColorFromTextureData(const uint8* MipDa
 		VertexColor.DWColor() &= ColorMask.DWColor();
 	}
 
+	// Vertex color is linear
+ 	if (Texture->SRGB)
+ 	{
+ 		VertexColor = FLinearColor(VertexColor).ToFColor(false);
+ 	}
+
 	return VertexColor;
 }
 
@@ -1482,6 +1532,8 @@ void UMeshPaintingSubsystem::CleanUp()
 	}
 	ComponentToAdapterMap.Empty();
 	FMeshPaintComponentAdapterFactory::CleanupGlobals();
+
+	CopiedTextureData.Init(FImageInfo());
 }
 
 bool UMeshPaintingSubsystem::FindHitResult(const FRay Ray, FHitResult& BestTraceResult)
@@ -1541,6 +1593,19 @@ void UMeshPaintingSubsystem::SetCopiedColorsByComponent(TArray<FPerComponentVert
 	CopiedColorsByComponent = InCopiedColors;
 }
 
+FImage const& UMeshPaintingSubsystem::GetCopiedTexture() const
+{
+	return CopiedTextureData;
+}
+
+void UMeshPaintingSubsystem::SetCopiedTexture(UTexture* InTexture)
+{
+	if (InTexture != nullptr)
+	{
+		InTexture->Source.GetMipImage(CopiedTextureData, 0);
+	}
+}
+
 void UMeshPaintingSubsystem::CacheSelectionData(const int32 PaintLODIndex, const int32 UVChannel)
 {
 	bSelectionContainsPerLODColors = false;
diff --git a/Engine/Plugins/MeshPainting/Source/MeshPaintingToolset/Private/MeshTexturePaintingTool.cpp b/Engine/Plugins/MeshPainting/Source/MeshPaintingToolset/Private/MeshTexturePaintingTool.cpp
--- a/Engine/Plugins/MeshPainting/Source/MeshPaintingToolset/Private/MeshTexturePaintingTool.cpp
+++ b/Engine/Plugins/MeshPainting/Source/MeshPaintingToolset/Private/MeshTexturePaintingTool.cpp
@@ -134,19 +134,29 @@ namespace UE::MeshPaintingToolset::Private::TexturePainting
 	}
 }
 
+
 /*
  * ToolBuilder
  */
 
-bool UMeshTexturePaintingToolBuilder::CanBuildTool(const FToolBuilderState& SceneState) const
+bool UMeshTextureColorPaintingToolBuilder::CanBuildTool(const FToolBuilderState& SceneState) const
 {
 	return GEngine->GetEngineSubsystem<UMeshPaintingSubsystem>()->SelectionHasMaterialValidForTexturePaint();
 }
 
-UInteractiveTool* UMeshTexturePaintingToolBuilder::BuildTool(const FToolBuilderState& SceneState) const
+UInteractiveTool* UMeshTextureColorPaintingToolBuilder::BuildTool(const FToolBuilderState& SceneState) const
 {
-	UMeshTexturePaintingTool* NewTool = NewObject<UMeshTexturePaintingTool>(SceneState.ToolManager);
-	return NewTool;
+	return NewObject<UMeshTextureColorPaintingTool>(SceneState.ToolManager);
+}
+
+bool UMeshTextureAssetPaintingToolBuilder::CanBuildTool(const FToolBuilderState& SceneState) const
+{
+	return GEngine->GetEngineSubsystem<UMeshPaintingSubsystem>()->SelectionHasMaterialValidForTexturePaint();
+}
+
+UInteractiveTool* UMeshTextureAssetPaintingToolBuilder::BuildTool(const FToolBuilderState& SceneState) const
+{
+	return NewObject<UMeshTextureAssetPaintingTool>(SceneState.ToolManager);
 }
 
 
@@ -154,25 +164,11 @@ UInteractiveTool* UMeshTexturePaintingToolBuilder::BuildTool(const FToolBuilderS
  * Tool
  */
 
-UMeshTexturePaintingToolProperties::UMeshTexturePaintingToolProperties()
-	:UBrushBaseProperties(),
-	PaintColor(FLinearColor::White),
-	EraseColor(FLinearColor::Black),
-	bWriteRed(true),
-	bWriteGreen(true),
-	bWriteBlue(true),
-	bWriteAlpha(false),
-	UVChannel(0),
-	bEnableSeamPainting(false),
-	PaintTexture(nullptr)
-{
-}
 UMeshTexturePaintingTool::UMeshTexturePaintingTool()
 {
 	PropertyClass = UMeshTexturePaintingToolProperties::StaticClass();
 }
 
-
 void UMeshTexturePaintingTool::Setup()
 {
 	Super::Setup();
@@ -195,13 +191,8 @@ void UMeshTexturePaintingTool::Setup()
 
 	BrushStampIndicator->LineColor = FLinearColor::Green;
 
-	GetToolManager()->DisplayMessage(
-		LOCTEXT("OnStartTexturePaintTool", "The Texture Weight Painting mode enables you to paint on textures and access available properties while doing so ."),
-		EToolMessageLevel::UserNotification);
-
 	SelectionMechanic = NewObject<UMeshPaintSelectionMechanic>(this);
 	SelectionMechanic->Setup(this);
-
 }
 
 
@@ -298,7 +289,6 @@ void UMeshTexturePaintingTool::OnTick(float DeltaTime)
 				BucketFillParams.bWriteGreen = TextureProperties->bWriteGreen;
 				BucketFillParams.bWriteBlue = TextureProperties->bWriteBlue;
 				BucketFillParams.bWriteAlpha = TextureProperties->bWriteAlpha;
-				BucketFillParams.UVChannel = TextureProperties->UVChannel;
 				BucketFillParams.bUseFillBucket = true;
 			}
 
@@ -306,13 +296,16 @@ void UMeshTexturePaintingTool::OnTick(float DeltaTime)
 			for (int32 j = 0; j < SelectedMeshComponents.Num(); ++j)
 			{
 				UMeshComponent* SelectedComponent = SelectedMeshComponents[j];
+				
+				const int32 UVChannel = GetSelectedUVChannel(SelectedComponent);
+
 				IMeshPaintComponentAdapter* MeshAdapter = MeshPaintingSubsystem->GetAdapterForComponent(SelectedComponent).Get();
 				if (MeshAdapter)
 				{
 					if (MeshAdapter->SupportsTexturePaint())
 					{
 						Textures.Empty();
-						const UTexture2D* TargetTexture2D = TextureProperties->PaintTexture;
+						const UTexture2D* TargetTexture2D = GetSelectedPaintTexture(SelectedComponent);
 						if (TargetTexture2D)
 						{
 							Textures.Add(TargetTexture2D);
@@ -327,7 +320,7 @@ void UMeshTexturePaintingTool::OnTick(float DeltaTime)
 							UTexturePaintToolset::RetrieveMeshSectionsForTextures(SelectedComponent, 0 /*CachedLODIndex*/, Textures, MaterialSections);
 
 							TArray<FTexturePaintTriangleInfo> TrianglePaintInfoArray;
-							FPerTrianglePaintAction TempAction = FPerTrianglePaintAction::CreateUObject(this, &UMeshTexturePaintingTool::GatherTextureTriangles, &TrianglePaintInfoArray, &MaterialSections, TextureProperties->UVChannel);
+							FPerTrianglePaintAction TempAction = FPerTrianglePaintAction::CreateUObject(this, &UMeshTexturePaintingTool::GatherTextureTriangles, &TrianglePaintInfoArray, &MaterialSections, UVChannel);
 
 							// We are flooding the texture, so all triangles are influenced
 
@@ -355,14 +348,15 @@ void UMeshTexturePaintingTool::OnTick(float DeltaTime)
 							}
 
 							FMeshPaintParameters* LastParams = nullptr;
-							PaintTexture(BucketFillParams, TrianglePaintInfoArray, *MeshAdapter, LastParams);
+							PaintTexture(BucketFillParams, UVChannel, TrianglePaintInfoArray, *MeshAdapter, LastParams);
 						}
 					}
 				}
 			}
 		}
 
-		if (MeshPaintingSubsystem->bNeedsRecache || (PaintableTextures.Num() > 0 && TextureProperties && TextureProperties->PaintTexture == nullptr))
+		UMeshComponent* FirstSelectedComponent = MeshPaintingSubsystem->GetSelectedMeshComponents()[0];
+		if (MeshPaintingSubsystem->bNeedsRecache || (PaintableTextures.Num() > 0 && TextureProperties && GetSelectedPaintTexture(FirstSelectedComponent) == nullptr))
 		{
 			CacheSelectionData();
 			CacheTexturePaintData();
@@ -478,34 +472,6 @@ void UMeshTexturePaintingTool::CacheSelectionData()
 	}
 }
 
-void UMeshTexturePaintingTool::CacheTexturePaintData()
-{
-	UMeshPaintingSubsystem* MeshPaintingSubsystem = GEngine->GetEngineSubsystem<UMeshPaintingSubsystem>();
-	if (MeshPaintingSubsystem)
-	{
-		TArray<UMeshComponent*> PaintableComponents = MeshPaintingSubsystem->GetPaintableMeshComponents();
-		
-		PaintableTextures.Empty();
-		if (PaintableComponents.Num() == 1 && PaintableComponents[0])
-		{
-			const UMeshComponent* Component = PaintableComponents[0];
-			TSharedPtr<IMeshPaintComponentAdapter> Adapter = MeshPaintingSubsystem->GetAdapterForComponent(Component);
-			UTexturePaintToolset::RetrieveTexturesForComponent(Component, Adapter.Get(), PaintableTextures);
-		}
-
-		// Ensure that the selection remains valid or is invalidated
-		if (!PaintableTextures.Contains(TextureProperties->PaintTexture))
-		{
-			UTexture2D* NewTexture = nullptr;
-			if (PaintableTextures.Num() > 0)
-			{
-				NewTexture = Cast<UTexture2D>(PaintableTextures[0].Texture);
-			}
-			TextureProperties->PaintTexture = NewTexture;
-		}
-	}
-}
-
 bool UMeshTexturePaintingTool::PaintInternal(const TArrayView<TPair<FVector, FVector>>& Rays, EMeshPaintModeAction PaintAction, float PaintStrength)
 {
 	TArray<FPaintRayResults> PaintRayResults;
@@ -549,15 +515,27 @@ bool UMeshTexturePaintingTool::PaintInternal(const TArrayView<TPair<FVector, FVe
 			}
 		}
 
-		bool bUsed = false;
+		UMeshComponent* BestTraceMeshComponent = Cast<UMeshComponent>(BestTraceResult.GetComponent());
+		// If painting texture assets, just use the BestTraceMeshComponent as we only support painting a single mesh at a time in that mode.
+		const bool bAllowMultiselect = AllowsMultiselect();
 
-		if (BestTraceResult.GetComponent() != nullptr)
+		bool bUsed = false;
+		for (UMeshComponent* MeshComponent : MeshPaintingSubsystem->GetPaintableMeshComponents())
 		{
-			// If we're using texture paint, just use the best trace result we found as we currently only
-			// support painting a single mesh at a time in that mode.
-			UMeshComponent* ComponentToPaint = CastChecked<UMeshComponent>(BestTraceResult.GetComponent());
-			HoveredComponents.FindOrAdd(ComponentToPaint).Add(i);
-			bUsed = true;
+			if (MeshComponent == BestTraceMeshComponent)
+			{
+				HoveredComponents.FindOrAdd(MeshComponent).Add(i);
+				bUsed = true;
+			}
+			else if (bAllowMultiselect)
+			{
+				FSphere Sphere(BestTraceResult.Location, BrushRadius);
+				if (MeshComponent->GetLocalBounds().GetSphere().TransformBy(MeshComponent->GetComponentTransform()).Intersects(Sphere))
+				{
+					HoveredComponents.FindOrAdd(MeshComponent).Add(i);
+					bUsed = true;
+				}
+			}
 		}
 
 		if (bUsed)
@@ -598,11 +576,6 @@ bool UMeshTexturePaintingTool::PaintInternal(const TArrayView<TPair<FVector, FVe
 				Params.bWriteAlpha = TextureProperties->bWriteAlpha;
 				FVector BrushSpaceVertexPosition = Params.InverseBrushToWorldMatrix.TransformVector(FVector4(Params.BrushPosition, 1.0f));
 				Params.BrushPosition2D = FVector2f(BrushSpaceVertexPosition.X, BrushSpaceVertexPosition.Y);
-	
-
-				// @todo MeshPaint: Ideally we would default to: TexturePaintingCurrentMeshComponent->StaticMesh->LightMapCoordinateIndex
-				//		Or we could indicate in the GUI which channel is the light map set (button to set it?)
-				Params.UVChannel = TextureProperties->UVChannel;
 			}
 		}
 	}
@@ -626,10 +599,12 @@ bool UMeshTexturePaintingTool::PaintInternal(const TArrayView<TPair<FVector, FVe
 				continue;
 			}
 
+			const int32 UVChannel = GetSelectedUVChannel(HoveredComponent);
+
 			if (MeshAdapter->SupportsTexturePaint())
 			{
 				Textures.Empty();
-				const UTexture2D* TargetTexture2D = TextureProperties->PaintTexture;
+				const UTexture2D* TargetTexture2D = GetSelectedPaintTexture(HoveredComponent);
 				if (TargetTexture2D)
 				{
 					Textures.Add(TargetTexture2D);
@@ -649,7 +624,7 @@ bool UMeshTexturePaintingTool::PaintInternal(const TArrayView<TPair<FVector, FVe
 						const FVector& BestTraceResultLocation = PaintRayResults[PaintRayResultId].BestTraceResult.Location;
 						FViewCameraState CameraState;
 						GetToolManager()->GetContextQueriesAPI()->GetCurrentViewState(CameraState);
-						bPaintApplied |= MeshPaintingSubsystem->ApplyPerTrianglePaintAction(MeshAdapter, CameraState.Position, BestTraceResultLocation, BrushProperties, FPerTrianglePaintAction::CreateUObject(this, &UMeshTexturePaintingTool::GatherTextureTriangles, &TrianglePaintInfoArray, &MaterialSections, TextureProperties->UVChannel), TextureProperties->bOnlyFrontFacingTriangles);
+						bPaintApplied |= MeshPaintingSubsystem->ApplyPerTrianglePaintAction(MeshAdapter, CameraState.Position, BestTraceResultLocation, BrushProperties, FPerTrianglePaintAction::CreateUObject(this, &UMeshTexturePaintingTool::GatherTextureTriangles, &TrianglePaintInfoArray, &MaterialSections, UVChannel), TextureProperties->bOnlyFrontFacingTriangles);
 						break;
 					}
 
@@ -676,7 +651,7 @@ bool UMeshTexturePaintingTool::PaintInternal(const TArrayView<TPair<FVector, FVe
 								LastParams = &LastPaintRayResults[PaintRayResultId].Params;
 							}
 
-							PaintTexture(Params, TrianglePaintInfoArray, *MeshAdapter, LastParams);
+							PaintTexture(Params, UVChannel, TrianglePaintInfoArray, *MeshAdapter, LastParams);
 							break;
 						}
 					}
@@ -753,8 +728,8 @@ FInputRayHit UMeshTexturePaintingTool::CanBeginClickDragSequence(const FInputDev
 	{
 		MeshPaintingSubsystem->LastPaintedComponent = (UMeshComponent*)LastBestHitResult.Component.Get();
 		GVertexViewModeOverrideOwnerName = *LastBestHitResult.Component->GetOwner()->GetName();
-		CommitAllPaintedTextures();
-		MeshPaintingSubsystem->bNeedsRecache = true;
+		//CommitAllPaintedTextures();
+		//MeshPaintingSubsystem->bNeedsRecache = true;
 	}
 
 	return Super::CanBeginClickDragSequence(PressPos);
@@ -896,7 +871,7 @@ void UMeshTexturePaintingTool::StartPaintingTexture(UMeshComponent* InMeshCompon
 
 	const auto FeatureLevel = InMeshComponent->GetWorld()->GetFeatureLevel();
 
-	UTexture2D* Texture2D = TextureProperties->PaintTexture;
+	UTexture2D* Texture2D = GetSelectedPaintTexture(InMeshComponent);
 	if (Texture2D == nullptr)
 	{
 		return;
@@ -970,56 +945,56 @@ void UMeshTexturePaintingTool::StartPaintingTexture(UMeshComponent* InMeshCompon
 			const int32 BrushTargetTextureHeight = TextureHeight;
 
 			// Create the rendertarget used to store our paint delta
-			if (BrushRenderTargetTexture == nullptr ||
-				BrushRenderTargetTexture->GetSurfaceWidth() != BrushTargetTextureWidth ||
-				BrushRenderTargetTexture->GetSurfaceHeight() != BrushTargetTextureHeight)
+			if (TextureData->BrushRenderTargetTexture == nullptr ||
+				TextureData->BrushRenderTargetTexture->GetSurfaceWidth() != BrushTargetTextureWidth ||
+				TextureData->BrushRenderTargetTexture->GetSurfaceHeight() != BrushTargetTextureHeight)
 			{
-				BrushRenderTargetTexture = nullptr;
-				BrushRenderTargetTexture = NewObject<UTextureRenderTarget2D>(GetTransientPackage(), NAME_None, RF_Transient);
+				TextureData->BrushRenderTargetTexture = nullptr;
+				TextureData->BrushRenderTargetTexture = NewObject<UTextureRenderTarget2D>(GetTransientPackage(), NAME_None, RF_Transient);
 				const bool bForceLinearGamma = true;
-				BrushRenderTargetTexture->ClearColor = FLinearColor::Black;
-				BrushRenderTargetTexture->bNeedsTwoCopies = true;
-				BrushRenderTargetTexture->InitCustomFormat(BrushTargetTextureWidth, BrushTargetTextureHeight, PF_A16B16G16R16, bForceLinearGamma);
-				BrushRenderTargetTexture->UpdateResourceImmediate();
-				BrushRenderTargetTexture->AddressX = TextureData->PaintRenderTargetTexture->AddressX;
-				BrushRenderTargetTexture->AddressY = TextureData->PaintRenderTargetTexture->AddressY;
+				TextureData->BrushRenderTargetTexture->ClearColor = FLinearColor::Black;
+				TextureData->BrushRenderTargetTexture->bNeedsTwoCopies = true;
+				TextureData->BrushRenderTargetTexture->InitCustomFormat(BrushTargetTextureWidth, BrushTargetTextureHeight, PF_A16B16G16R16, bForceLinearGamma);
+				TextureData->BrushRenderTargetTexture->UpdateResourceImmediate();
+				TextureData->BrushRenderTargetTexture->AddressX = TextureData->PaintRenderTargetTexture->AddressX;
+				TextureData->BrushRenderTargetTexture->AddressY = TextureData->PaintRenderTargetTexture->AddressY;
 			}
 
 			if (TextureProperties->bEnableSeamPainting)
 			{
 				// Create the rendertarget used to store a mask for our paint delta area 
-				if (BrushMaskRenderTargetTexture == nullptr ||
-					BrushMaskRenderTargetTexture->GetSurfaceWidth() != BrushTargetTextureWidth ||
-					BrushMaskRenderTargetTexture->GetSurfaceHeight() != BrushTargetTextureHeight)
+				if (TextureData->BrushMaskRenderTargetTexture == nullptr ||
+					TextureData->BrushMaskRenderTargetTexture->GetSurfaceWidth() != BrushTargetTextureWidth ||
+					TextureData->BrushMaskRenderTargetTexture->GetSurfaceHeight() != BrushTargetTextureHeight)
 				{
-					BrushMaskRenderTargetTexture = nullptr;
-					BrushMaskRenderTargetTexture = NewObject<UTextureRenderTarget2D>(GetTransientPackage(), NAME_None, RF_Transient);
+					TextureData->BrushMaskRenderTargetTexture = nullptr;
+					TextureData->BrushMaskRenderTargetTexture = NewObject<UTextureRenderTarget2D>(GetTransientPackage(), NAME_None, RF_Transient);
 					const bool bForceLinearGamma = true;
-					BrushMaskRenderTargetTexture->ClearColor = FLinearColor::Black;
-					BrushMaskRenderTargetTexture->bNeedsTwoCopies = true;
-					BrushMaskRenderTargetTexture->InitCustomFormat(BrushTargetTextureWidth, BrushTargetTextureHeight, PF_B8G8R8A8, bForceLinearGamma);
-					BrushMaskRenderTargetTexture->UpdateResourceImmediate();
-					BrushMaskRenderTargetTexture->AddressX = TextureData->PaintRenderTargetTexture->AddressX;
-					BrushMaskRenderTargetTexture->AddressY = TextureData->PaintRenderTargetTexture->AddressY;
+					TextureData->BrushMaskRenderTargetTexture->ClearColor = FLinearColor::Black;
+					TextureData->BrushMaskRenderTargetTexture->bNeedsTwoCopies = true;
+					TextureData->BrushMaskRenderTargetTexture->InitCustomFormat(BrushTargetTextureWidth, BrushTargetTextureHeight, PF_B8G8R8A8, bForceLinearGamma);
+					TextureData->BrushMaskRenderTargetTexture->UpdateResourceImmediate();
+					TextureData->BrushMaskRenderTargetTexture->AddressX = TextureData->PaintRenderTargetTexture->AddressX;
+					TextureData->BrushMaskRenderTargetTexture->AddressY = TextureData->PaintRenderTargetTexture->AddressY;
 				}
 
 				// Create the rendertarget used to store a texture seam mask
-				if (SeamMaskRenderTargetTexture == nullptr ||
-					SeamMaskRenderTargetTexture->GetSurfaceWidth() != TextureWidth ||
-					SeamMaskRenderTargetTexture->GetSurfaceHeight() != TextureHeight)
+				if (TextureData->SeamMaskRenderTargetTexture == nullptr ||
+					TextureData->SeamMaskRenderTargetTexture->GetSurfaceWidth() != TextureWidth ||
+					TextureData->SeamMaskRenderTargetTexture->GetSurfaceHeight() != TextureHeight)
 				{
-					SeamMaskRenderTargetTexture = nullptr;
-					SeamMaskRenderTargetTexture = NewObject<UTextureRenderTarget2D>(GetTransientPackage(), NAME_None, RF_Transient);
+					TextureData->SeamMaskRenderTargetTexture = nullptr;
+					TextureData->SeamMaskRenderTargetTexture = NewObject<UTextureRenderTarget2D>(GetTransientPackage(), NAME_None, RF_Transient);
 					const bool bForceLinearGamma = true;
-					SeamMaskRenderTargetTexture->ClearColor = FLinearColor::Black;
-					SeamMaskRenderTargetTexture->bNeedsTwoCopies = true;
-					SeamMaskRenderTargetTexture->InitCustomFormat(BrushTargetTextureWidth, BrushTargetTextureHeight, PF_B8G8R8A8, bForceLinearGamma);
-					SeamMaskRenderTargetTexture->UpdateResourceImmediate();
-					SeamMaskRenderTargetTexture->AddressX = TextureData->PaintRenderTargetTexture->AddressX;
-					SeamMaskRenderTargetTexture->AddressY = TextureData->PaintRenderTargetTexture->AddressY;
+					TextureData->SeamMaskRenderTargetTexture->ClearColor = FLinearColor::Black;
+					TextureData->SeamMaskRenderTargetTexture->bNeedsTwoCopies = true;
+					TextureData->SeamMaskRenderTargetTexture->InitCustomFormat(BrushTargetTextureWidth, BrushTargetTextureHeight, PF_B8G8R8A8, bForceLinearGamma);
+					TextureData->SeamMaskRenderTargetTexture->UpdateResourceImmediate();
+					TextureData->SeamMaskRenderTargetTexture->AddressX = TextureData->PaintRenderTargetTexture->AddressX;
+					TextureData->SeamMaskRenderTargetTexture->AddressY = TextureData->PaintRenderTargetTexture->AddressY;
 				}
 
-				bGenerateSeamMask = true;
+				TextureData->bGenerateSeamMask = true;
 			}
 
 			// MeshPaint: Here we override the textures on the mesh with the render target.  The problem is that other meshes in the scene that use
@@ -1079,7 +1054,7 @@ void UMeshTexturePaintingTool::StartPaintingTexture(UMeshComponent* InMeshCompon
 	}
 }
 
-void UMeshTexturePaintingTool::PaintTexture(FMeshPaintParameters& InParams, TArray<FTexturePaintTriangleInfo>& InInfluencedTriangles, const IMeshPaintComponentAdapter& GeometryInfo, FMeshPaintParameters* LastParams)
+void UMeshTexturePaintingTool::PaintTexture(FMeshPaintParameters& InParams, int32 UVChannel, TArray<FTexturePaintTriangleInfo>& InInfluencedTriangles, const IMeshPaintComponentAdapter& GeometryInfo, FMeshPaintParameters* LastParams)
 {
 	// We bail early if there are no influenced triangles
 	if (InInfluencedTriangles.Num() <= 0)
@@ -1096,8 +1071,8 @@ void UMeshTexturePaintingTool::PaintTexture(FMeshPaintParameters& InParams, TArr
 
 	// Copy the current image to the brush rendertarget texture.
 	{
-		check(BrushRenderTargetTexture != nullptr);
-		UTexturePaintToolset::CopyTextureToRenderTargetTexture(TextureData->PaintRenderTargetTexture, BrushRenderTargetTexture, FeatureLevel);
+		check(TextureData->BrushRenderTargetTexture != nullptr);
+		UTexturePaintToolset::CopyTextureToRenderTargetTexture(TextureData->PaintRenderTargetTexture, TextureData->BrushRenderTargetTexture, FeatureLevel);
 	}
 
 	const bool bEnableSeamPainting = TextureProperties->bEnableSeamPainting;
@@ -1106,7 +1081,7 @@ void UMeshTexturePaintingTool::PaintTexture(FMeshPaintParameters& InParams, TArr
 	// Grab the actual render target resource from the textures.  Note that we're absolutely NOT ALLOWED to
 	// dereference these pointers.  We're just passing them along to other functions that will use them on the render
 	// thread.  The only thing we're allowed to do is check to see if they are nullptr or not.
-	FTextureRenderTargetResource* BrushRenderTargetResource = BrushRenderTargetTexture->GameThread_GetRenderTargetResource();
+	FTextureRenderTargetResource* BrushRenderTargetResource = TextureData->BrushRenderTargetTexture->GameThread_GetRenderTargetResource();
 	check(BrushRenderTargetResource != nullptr);
 
 	// Create a canvas for the brush render target.
@@ -1128,7 +1103,7 @@ void UMeshTexturePaintingTool::PaintTexture(FMeshPaintParameters& InParams, TArr
 		}
 		MeshPaintBatchedElementParameters->ShaderParams.PaintBrushRotationOffset = TextureProperties->PaintBrushRotationOffset;
 		MeshPaintBatchedElementParameters->ShaderParams.bUseFillBucket = InParams.bUseFillBucket;
-		MeshPaintBatchedElementParameters->ShaderParams.CloneTexture = BrushRenderTargetTexture;
+		MeshPaintBatchedElementParameters->ShaderParams.CloneTexture = TextureData->BrushRenderTargetTexture;
 		MeshPaintBatchedElementParameters->ShaderParams.WorldToBrushMatrix = WorldToBrushMatrix;
 		MeshPaintBatchedElementParameters->ShaderParams.BrushRadius = InParams.InnerBrushRadius + InParams.BrushRadialFalloffRange;
 		MeshPaintBatchedElementParameters->ShaderParams.BrushRadialFalloffRange = InParams.BrushRadialFalloffRange;
@@ -1157,7 +1132,7 @@ void UMeshTexturePaintingTool::PaintTexture(FMeshPaintParameters& InParams, TArr
 
 	if (bEnableSeamPainting)
 	{
-		BrushMaskRenderTargetResource = BrushMaskRenderTargetTexture->GameThread_GetRenderTargetResource();
+		BrushMaskRenderTargetResource = TextureData->BrushMaskRenderTargetTexture->GameThread_GetRenderTargetResource();
 		check(BrushMaskRenderTargetResource != nullptr);
 
 		// Create a canvas for the brush mask rendertarget and clear it to black.
@@ -1327,11 +1302,10 @@ void UMeshTexturePaintingTool::PaintTexture(FMeshPaintParameters& InParams, TArr
 	if (!bEnableSeamPainting)
 	{
 		// Seam painting is not enabled so we just copy our delta paint info to the paint target.
-		UTexturePaintToolset::CopyTextureToRenderTargetTexture(BrushRenderTargetTexture, TextureData->PaintRenderTargetTexture, FeatureLevel);
+		UTexturePaintToolset::CopyTextureToRenderTargetTexture(TextureData->BrushRenderTargetTexture, TextureData->PaintRenderTargetTexture, FeatureLevel);
 	}
 	else
 	{
-
 		// Constants used for generating quads across entire paint rendertarget
 		const float MinU = 0.0f;
 		const float MinV = 0.0f;
@@ -1342,15 +1316,11 @@ void UMeshTexturePaintingTool::PaintTexture(FMeshPaintParameters& InParams, TArr
 		const float MaxX = TextureData->PaintRenderTargetTexture->GetSurfaceWidth();
 		const float MaxY = TextureData->PaintRenderTargetTexture->GetSurfaceHeight();
 
-		if (bGenerateSeamMask == true)
+		if (TextureData->bGenerateSeamMask == true)
 		{
-			// Generate the texture seam mask.  This is a slow operation when the object has many triangles so we only do it
-			//  once when painting is started.
-
-			FPaintTexture2DData* SeamTextureData = GetPaintTargetData(TextureProperties->PaintTexture);
-
-			UTexturePaintToolset::GenerateSeamMask(TexturePaintingCurrentMeshComponent, InParams.UVChannel, SeamMaskRenderTargetTexture, TextureProperties->PaintTexture, SeamTextureData != nullptr ? SeamTextureData->PaintRenderTargetTexture : nullptr);
-			bGenerateSeamMask = false;
+			// Generate the texture seam mask.  This is a slow operation when the object has many triangles so we only do it once when painting is started.
+			UTexturePaintToolset::GenerateSeamMask(TexturePaintingCurrentMeshComponent, UVChannel, TextureData->SeamMaskRenderTargetTexture, TextureData->PaintingTexture2D, TextureData->PaintRenderTargetTexture);
+			TextureData->bGenerateSeamMask = false;
 		}
 
 		FTextureRenderTargetResource* RenderTargetResource = TextureData->PaintRenderTargetTexture->GameThread_GetRenderTargetResource();
@@ -1363,9 +1333,9 @@ void UMeshTexturePaintingTool::PaintTexture(FMeshPaintParameters& InParams, TArr
 
 			TRefCountPtr< FMeshPaintDilateBatchedElementParameters > MeshPaintDilateBatchedElementParameters(new FMeshPaintDilateBatchedElementParameters());
 			{
-				MeshPaintDilateBatchedElementParameters->ShaderParams.Texture0 = BrushRenderTargetTexture;
-				MeshPaintDilateBatchedElementParameters->ShaderParams.Texture1 = SeamMaskRenderTargetTexture;
-				MeshPaintDilateBatchedElementParameters->ShaderParams.Texture2 = BrushMaskRenderTargetTexture;
+				MeshPaintDilateBatchedElementParameters->ShaderParams.Texture0 = TextureData->BrushRenderTargetTexture;
+				MeshPaintDilateBatchedElementParameters->ShaderParams.Texture1 = TextureData->SeamMaskRenderTargetTexture;
+				MeshPaintDilateBatchedElementParameters->ShaderParams.Texture2 = TextureData->BrushMaskRenderTargetTexture;
 				MeshPaintDilateBatchedElementParameters->ShaderParams.WidthPixelOffset = (float)(1.0f / TextureData->PaintRenderTargetTexture->GetSurfaceWidth());
 				MeshPaintDilateBatchedElementParameters->ShaderParams.HeightPixelOffset = (float)(1.0f / TextureData->PaintRenderTargetTexture->GetSurfaceHeight());
 
@@ -1408,7 +1378,6 @@ void UMeshTexturePaintingTool::PaintTexture(FMeshPaintParameters& InParams, TArr
 
 			// Tell the rendering thread to draw any remaining batched elements
 			Canvas3.Flush_GameThread(true);
-
 		}
 
 		ENQUEUE_RENDER_COMMAND(UpdateMeshPaintRTCommand3)(
@@ -1592,33 +1561,14 @@ void UMeshTexturePaintingTool::OnTransactionStateChanged(const FTransactionConte
 						CompareSortedArrayElements<TObjectPtr<UTexture2D>>(CurrentTextures, OverrideTextures, TextureOverrideIsNotNeeded, TextureOverrideIsMissing, UpdateOverridedComponents);
 					}
 
-
 					// Update the render targets
+					for (TPair<TObjectPtr<UTexture2D>, FPaintTexture2DData>& TargetData : PaintTargetData)
 					{
-						if (Transaction->ContainsObject(this) && TextureProperties)
-						{
-							if (TextureProperties->PaintTexture)
-							{
-								if (FPaintTexture2DData* TextureData = GetPaintTargetData(TextureProperties->PaintTexture))
-								{
-									if (TextureData->ScratchTexture)
-									{
-										UTexturePaintToolset::UpdateRenderTargetData(*TextureData);
-									}
-								}
-							}
-						}
-						else
-						{
-							for (TPair<TObjectPtr<UTexture2D>, FPaintTexture2DData>& TagetData : PaintTargetData)
-							{
-								FPaintTexture2DData& PaintTexture2DData = TagetData.Value;
+						FPaintTexture2DData& PaintTexture2DData = TargetData.Value;
 
-								if (Transaction->ContainsObject(PaintTexture2DData.ScratchTexture))
-								{
-									UTexturePaintToolset::UpdateRenderTargetData(PaintTexture2DData);
-								}
-							}
+						if (Transaction->ContainsObject(PaintTexture2DData.ScratchTexture))
+						{
+							UTexturePaintToolset::UpdateRenderTargetData(PaintTexture2DData);
 						}
 					}
 				}
@@ -1627,13 +1577,13 @@ void UMeshTexturePaintingTool::OnTransactionStateChanged(const FTransactionConte
 	}
 }
 
-void UMeshTexturePaintingTool::CycleTextures(int32 Direction)
+void UMeshTextureAssetPaintingTool::CycleTextures(int32 Direction)
 {
 	if (!PaintableTextures.Num())
 	{
 		return;
 	}
-	TObjectPtr<UTexture2D>& SelectedTexture = TextureProperties->PaintTexture;
+	TObjectPtr<UTexture2D>& SelectedTexture = AssetProperties->PaintTexture;
 	const int32 TextureIndex = (SelectedTexture != nullptr) ? PaintableTextures.IndexOfByKey(SelectedTexture) : 0;
 	if (TextureIndex != INDEX_NONE)
 	{
@@ -1704,10 +1654,6 @@ void UMeshTexturePaintingTool::CommitAllPaintedTextures()
 					FMemory::Memcpy(Colors, TexturePixels.GetData(), TexturePixels.Num() * sizeof(FColor));
 					TextureData->PaintingTexture2D->Source.UnlockMip(0);
 
-					// If render target gamma used was 1.0 then disable SRGB for the static texture
-					// @todo MeshPaint: We are not allowed to dereference the RenderTargetResource pointer, figure out why we need this when the GetDisplayGamma() function is hard coded to return 2.2.
-					TextureData->PaintingTexture2D->SRGB = FMath::Abs(RenderTargetResource->GetDisplayGamma() - 1.0f) >= KINDA_SMALL_NUMBER;
-
 					TextureData->PaintingTexture2D->bHasBeenPaintedInEditor = true;
 
 					// Update the texture (generate mips, compress if needed)
@@ -1767,7 +1713,107 @@ int32 UMeshTexturePaintingTool::GetNumberOfPendingPaintChanges() const
 	return Result;
 }
 
-bool UMeshTexturePaintingTool::ShouldFilterTextureAsset(const FAssetData& AssetData) const
+bool UMeshTexturePaintingTool::IsMeshAdapterSupported(TSharedPtr<IMeshPaintComponentAdapter> MeshAdapter) const
+{
+	return MeshAdapter.IsValid() ? MeshAdapter->SupportsTexturePaint() : false;
+}
+
+void UMeshTexturePaintingTool::FloodCurrentPaintTexture()
+{
+	bRequestPaintBucketFill = true;
+}
+
+
+UMeshTextureColorPaintingTool::UMeshTextureColorPaintingTool()
+{
+	PropertyClass = UMeshTextureColorPaintingToolProperties::StaticClass();
+}
+
+void UMeshTextureColorPaintingTool::Setup()
+{
+	Super::Setup();
+	ColorProperties = Cast<UMeshTextureColorPaintingToolProperties>(BrushProperties);
+
+	GetToolManager()->DisplayMessage(
+		LOCTEXT("OnStartTextureColorPaintTool", "Paint colors to the Mesh Paint Texture object stored on mesh components."),
+		EToolMessageLevel::UserNotification);
+}
+
+UTexture2D* UMeshTextureColorPaintingTool::GetSelectedPaintTexture(UMeshComponent const* InMeshComponent) const
+{
+	return Cast<UTexture2D>(InMeshComponent->GetMeshPaintTexture());
+}
+
+int32 UMeshTextureColorPaintingTool::GetSelectedUVChannel(UMeshComponent const* InMeshComponent) const
+{
+	return InMeshComponent != nullptr ? InMeshComponent->GetMeshPaintTextureCoordinateIndex() : 0;
+}
+
+void UMeshTextureColorPaintingTool::GetModifiedTexturesToSave(TArray<UObject*>& OutTexturesToSave) const
+{
+	for (FPaintableTexture const& PaintableTexture : PaintableTextures)
+	{
+		if (PaintableTexture.Texture->GetOutermost()->IsDirty())
+		{
+			OutTexturesToSave.Add(PaintableTexture.Texture);
+		}
+	}
+}
+
+void UMeshTextureColorPaintingTool::CacheTexturePaintData()
+{
+	UMeshPaintingSubsystem* MeshPaintingSubsystem = GEngine->GetEngineSubsystem<UMeshPaintingSubsystem>();
+	if (MeshPaintingSubsystem)
+	{
+		PaintableTextures.Empty();
+
+		TArray<UMeshComponent*> PaintableComponents = MeshPaintingSubsystem->GetPaintableMeshComponents();
+		for (UMeshComponent const* Component : PaintableComponents)
+		{
+			int32 DummyDefaultIndex;
+			TSharedPtr<IMeshPaintComponentAdapter> Adapter = MeshPaintingSubsystem->GetAdapterForComponent(Component);
+			UTexturePaintToolset::RetrieveTexturesForComponent(Component, Adapter.Get(), DummyDefaultIndex, PaintableTextures);
+		}
+
+		PaintableTextures.RemoveAll([](FPaintableTexture const& PaintableTexture) { return !PaintableTexture.bIsMeshTexture; });
+	}
+}
+
+
+UMeshTextureAssetPaintingTool::UMeshTextureAssetPaintingTool()
+{
+	PropertyClass = UMeshTextureAssetPaintingToolProperties::StaticClass();
+}
+
+void UMeshTextureAssetPaintingTool::Setup()
+{
+	Super::Setup();
+	AssetProperties = Cast<UMeshTextureAssetPaintingToolProperties>(BrushProperties);
+
+	GetToolManager()->DisplayMessage(
+		LOCTEXT("OnStartTexturePaintTool", "The Texture Weight Painting mode enables you to paint on textures and access available properties while doing so ."),
+		EToolMessageLevel::UserNotification);
+}
+
+UTexture2D* UMeshTextureAssetPaintingTool::GetSelectedPaintTexture(UMeshComponent const* InMeshComponent) const
+{
+	return AssetProperties->PaintTexture;
+}
+
+int32 UMeshTextureAssetPaintingTool::GetSelectedUVChannel(UMeshComponent const* InMeshComponent) const
+{
+	return AssetProperties->UVChannel;
+}
+
+void UMeshTextureAssetPaintingTool::GetModifiedTexturesToSave(TArray<UObject*>& OutTexturesToSave) const
+{
+	if (AssetProperties->PaintTexture != nullptr && AssetProperties->PaintTexture->GetOutermost()->IsDirty())
+	{
+		OutTexturesToSave.Add(AssetProperties->PaintTexture);
+	}
+}
+
+bool UMeshTextureAssetPaintingTool::ShouldFilterTextureAsset(const FAssetData& AssetData) const
 {
 	UTexture2D* AssetTexture = Cast<UTexture2D>(AssetData.GetAsset());
 	if (!AssetTexture)
@@ -1777,7 +1823,7 @@ bool UMeshTexturePaintingTool::ShouldFilterTextureAsset(const FAssetData& AssetD
 	return !(PaintableTextures.ContainsByPredicate([=](const FPaintableTexture& Texture) { return Texture.Texture->GetFullName() == AssetTexture->GetFullName(); }));
 }
 
-void UMeshTexturePaintingTool::PaintTextureChanged(const FAssetData& AssetData)
+void UMeshTextureAssetPaintingTool::PaintTextureChanged(const FAssetData& AssetData)
 {
 	UTexture2D* Texture = Cast<UTexture2D>(AssetData.GetAsset());
 	if (Texture)
@@ -1789,7 +1835,7 @@ void UMeshTexturePaintingTool::PaintTextureChanged(const FAssetData& AssetData)
 			if (TextureTarget.TextureData == Texture)
 			{
 				TextureTarget.bIsSelected = true;
-				TextureProperties->UVChannel = TextureTarget.UVChannelIndex;
+				AssetProperties->UVChannel = TextureTarget.UVChannelIndex;
 			}
 			else
 			{
@@ -1799,14 +1845,44 @@ void UMeshTexturePaintingTool::PaintTextureChanged(const FAssetData& AssetData)
 	}
 }
 
-bool UMeshTexturePaintingTool::IsMeshAdapterSupported(TSharedPtr<IMeshPaintComponentAdapter> MeshAdapter) const
+void UMeshTextureAssetPaintingTool::CacheTexturePaintData()
 {
-	return MeshAdapter.IsValid() ? MeshAdapter->SupportsTexturePaint() : false;
-}
+	UMeshPaintingSubsystem* MeshPaintingSubsystem = GEngine->GetEngineSubsystem<UMeshPaintingSubsystem>();
+	if (MeshPaintingSubsystem)
+	{
+		PaintableTextures.Empty();
+		
+		UTexture* DefaultTexture = nullptr;
+	
+		TArray<UMeshComponent*> PaintableComponents = MeshPaintingSubsystem->GetPaintableMeshComponents();
+		if (PaintableComponents.IsValidIndex(0))
+		{
+			int32 DefaultTextureIndex = INDEX_NONE;
+			TSharedPtr<IMeshPaintComponentAdapter> Adapter = MeshPaintingSubsystem->GetAdapterForComponent(PaintableComponents[0]);
+			UTexturePaintToolset::RetrieveTexturesForComponent(PaintableComponents[0], Adapter.Get(), DefaultTextureIndex, PaintableTextures);
+			if (PaintableTextures.IsValidIndex(DefaultTextureIndex))
+			{
+				DefaultTexture = PaintableTextures[DefaultTextureIndex].Texture;
+			}
+		}
 
-void UMeshTexturePaintingTool::FloodCurrentPaintTexture()
-{
-	bRequestPaintBucketFill = true;
+		PaintableTextures.RemoveAll([](FPaintableTexture const& PaintableTexture) { return PaintableTexture.bIsMeshTexture; });
+
+		// Ensure that the selection remains valid or is invalidated
+		if (!PaintableTextures.Contains(AssetProperties->PaintTexture))
+		{
+			UTexture2D* NewTexture = nullptr;
+			if (PaintableTextures.Contains(DefaultTexture))
+			{
+				NewTexture = Cast<UTexture2D>(DefaultTexture);
+			}
+			else if (PaintableTextures.Num() > 0)
+			{
+				NewTexture = Cast<UTexture2D>(PaintableTextures[0].Texture);
+			}
+			AssetProperties->PaintTexture = NewTexture;
+		}
+	}
 }
 
 #undef LOCTEXT_NAMESPACE
diff --git a/Engine/Plugins/MeshPainting/Source/MeshPaintingToolset/Private/MeshVertexPaintingTool.cpp b/Engine/Plugins/MeshPainting/Source/MeshPaintingToolset/Private/MeshVertexPaintingTool.cpp
--- a/Engine/Plugins/MeshPainting/Source/MeshPaintingToolset/Private/MeshVertexPaintingTool.cpp
+++ b/Engine/Plugins/MeshPainting/Source/MeshPaintingToolset/Private/MeshVertexPaintingTool.cpp
@@ -46,10 +46,7 @@ UInteractiveTool* UMeshWeightPaintingToolBuilder::BuildTool(const FToolBuilderSt
  */
 
 UMeshVertexPaintingToolProperties::UMeshVertexPaintingToolProperties()
-	:UBrushBaseProperties(),
-	PaintColor(FLinearColor::White),
-	EraseColor(FLinearColor::Black),
-	bEnableFlow(false),
+	:UMeshPaintingToolProperties(),
 	VertexPreviewSize(6.0f)
 {
 }
@@ -522,25 +519,16 @@ void UMeshVertexPaintingTool::FinishPainting()
 	}
 }
 
-UMeshColorPaintingToolProperties::UMeshColorPaintingToolProperties()
-	:UMeshVertexPaintingToolProperties(),
-	bWriteRed(true),
-	bWriteGreen(true),
-	bWriteBlue(true),
-	bWriteAlpha(false)
-{
-
-}
 
 UMeshColorPaintingTool::UMeshColorPaintingTool()
 {
-	PropertyClass = UMeshColorPaintingToolProperties::StaticClass();
+	PropertyClass = UMeshVertexColorPaintingToolProperties::StaticClass();
 }
 
 void UMeshColorPaintingTool::Setup()
 {
 	Super::Setup();
-	ColorProperties = Cast<UMeshColorPaintingToolProperties>(BrushProperties);
+	ColorProperties = Cast<UMeshVertexColorPaintingToolProperties>(BrushProperties);
 
 	GetToolManager()->DisplayMessage(
 		LOCTEXT("OnStartColorPaintTool", "Paint vertex colors on selected meshes.  Use the Color View Mode to preview your applied changes."),
@@ -699,7 +687,7 @@ void UMeshColorPaintingTool::CycleMeshLODs(int32 Direction)
 	}
 }
 
-UMeshWeightPaintingToolProperties::UMeshWeightPaintingToolProperties()
+UMeshVertexWeightPaintingToolProperties::UMeshVertexWeightPaintingToolProperties()
 	:UMeshVertexPaintingToolProperties(),
 	TextureWeightType(EMeshPaintWeightTypes::AlphaLerp),
 	PaintTextureWeightIndex(EMeshPaintTextureIndex::TextureOne),
@@ -710,13 +698,13 @@ UMeshWeightPaintingToolProperties::UMeshWeightPaintingToolProperties()
 
 UMeshWeightPaintingTool::UMeshWeightPaintingTool()
 {
-	PropertyClass = UMeshWeightPaintingToolProperties::StaticClass();
+	PropertyClass = UMeshVertexWeightPaintingToolProperties::StaticClass();
 }
 
 void UMeshWeightPaintingTool::Setup()
 {
 	Super::Setup();
-	WeightProperties = Cast<UMeshWeightPaintingToolProperties>(BrushProperties);
+	WeightProperties = Cast<UMeshVertexWeightPaintingToolProperties>(BrushProperties);
 
 	GetToolManager()->DisplayMessage(
 		LOCTEXT("OnStartPaintWeightsTool", "Paint Vertex Weights on selected meshes."),
diff --git a/Engine/Plugins/MeshPainting/Source/MeshPaintingToolset/Private/TexturePaintToolset.cpp b/Engine/Plugins/MeshPainting/Source/MeshPaintingToolset/Private/TexturePaintToolset.cpp
--- a/Engine/Plugins/MeshPainting/Source/MeshPaintingToolset/Private/TexturePaintToolset.cpp
+++ b/Engine/Plugins/MeshPainting/Source/MeshPaintingToolset/Private/TexturePaintToolset.cpp
@@ -298,13 +298,16 @@ bool UTexturePaintToolset::GenerateSeamMask(UMeshComponent* MeshComponent, int32
 			EachTri.V2_Color = FLinearColor::Black;
 			TriList.Add(EachTri);
 		}
-		// Setup the tri render item with the list of tris
-		FCanvasTriangleItem TriItem(TriList, RenderTargetResource);
-		TriItem.BlendMode = SE_BLEND_Opaque;
-		// And render it
-		Canvas.DrawItem(TriItem);
-		// Tell the rendering thread to draw any remaining batched elements
-		Canvas.Flush_GameThread(true);
+		if (TriList.Num())
+		{
+			// Setup the tri render item with the list of tris
+			FCanvasTriangleItem TriItem(TriList, RenderTargetResource);
+			TriItem.BlendMode = SE_BLEND_Opaque;
+			// And render it
+			Canvas.DrawItem(TriItem);
+			// Tell the rendering thread to draw any remaining batched elements
+			Canvas.Flush_GameThread(true);
+		}
 	}
 
 
@@ -563,8 +566,10 @@ bool UTexturePaintToolset::DoesMeshComponentUseTexture(UMeshComponent* MeshCompo
 	return UsedTextures.Contains(Texture);
 }
 
-void UTexturePaintToolset::RetrieveTexturesForComponent(const UMeshComponent* Component, IMeshPaintComponentAdapter* Adapter, TArray<FPaintableTexture>& OutTextures)
+void UTexturePaintToolset::RetrieveTexturesForComponent(const UMeshComponent* Component, IMeshPaintComponentAdapter* Adapter, int32& OutDefaultIndex, TArray<FPaintableTexture>& OutTextures)
 {
+	OutDefaultIndex = INDEX_NONE;
+
 	if (Component && Adapter)
 	{
 		// Get the materials used by the mesh
@@ -573,21 +578,27 @@ void UTexturePaintToolset::RetrieveTexturesForComponent(const UMeshComponent* Co
 
 		for (int32 MaterialIndex = 0; MaterialIndex < UsedMaterials.Num(); ++MaterialIndex)
 		{
-			int32 OutDefaultIndex = 0;
-			Adapter->QueryPaintableTextures(MaterialIndex, OutDefaultIndex, OutTextures);
+			int32 OutDefaultIndexForMaterial = INDEX_NONE;
+			Adapter->QueryPaintableTextures(MaterialIndex, OutDefaultIndexForMaterial, OutTextures);
+
+			// We can only collect one default texture from the multiple materials!
+			if (OutDefaultIndex == INDEX_NONE && OutDefaultIndexForMaterial != INDEX_NONE)
+			{
+				OutDefaultIndex = OutDefaultIndexForMaterial;
+			}
 		}
 
 		// Filter out any paintable texture that is not supported currently (Todo: should we do this here or could we just move some logic in the FPaintableTexture)?
 		OutTextures.RemoveAll([](const FPaintableTexture& PaintableTexture)
+		{
+			if (const UTexture* Texture = PaintableTexture.Texture)
 			{
-				if (const UTexture* Texture = PaintableTexture.Texture)
-				{
-					// Only BRGA8 image are supported by the tool currently (we should probably provide some feedback to the user on why some texture are not supported and a way to convert them)
-					return Texture->Source.GetFormat() != TSF_BGRA8;
-				}
+				// Only BRGA8 image are supported by the tool currently (we should probably provide some feedback to the user on why some texture are not supported and a way to convert them)
+				return Texture->Source.GetFormat() != TSF_BGRA8;
+			}
 
-				return true;
-			});
+			return true;
+		});
 	}
 }
 
diff --git a/Engine/Plugins/MeshPainting/Source/MeshPaintingToolset/Public/BaseMeshPaintingToolProperties.h b/Engine/Plugins/MeshPainting/Source/MeshPaintingToolset/Public/BaseMeshPaintingToolProperties.h
new file mode 100644
--- /dev/null
+++ b/Engine/Plugins/MeshPainting/Source/MeshPaintingToolset/Public/BaseMeshPaintingToolProperties.h
@@ -0,0 +1,29 @@
+// Copyright Epic Games, Inc. All Rights Reserved.
+
+#pragma once
+
+#include "BaseTools/BaseBrushTool.h"
+#include "BaseMeshPaintingToolProperties.generated.h"
+
+UCLASS()
+class MESHPAINTINGTOOLSET_API UMeshPaintingToolProperties : public UBrushBaseProperties
+{
+	GENERATED_BODY()
+
+public:
+	/** Color used for applying color painting */
+	UPROPERTY(EditAnywhere, Category = Brush)
+	FLinearColor PaintColor = FLinearColor::White;
+
+	/** Color used for erasing color painting */
+	UPROPERTY(EditAnywhere, Category = Brush)
+	FLinearColor EraseColor = FLinearColor::Black;
+
+	/** Enables "Flow" painting where paint is continually applied from the brush every tick */
+	UPROPERTY(EditAnywhere, Category = Brush, meta = (DisplayName = "Enable Brush Flow"))
+	bool bEnableFlow = false;
+
+	/** Whether back-facing triangles should be ignored */
+	UPROPERTY(EditAnywhere, Category = Brush, meta = (DisplayName = "Ignore Back-Facing"))
+	bool bOnlyFrontFacingTriangles = false;
+};
diff --git a/Engine/Plugins/MeshPainting/Source/MeshPaintingToolset/Public/MeshPaintHelpers.h b/Engine/Plugins/MeshPainting/Source/MeshPaintingToolset/Public/MeshPaintHelpers.h
--- a/Engine/Plugins/MeshPainting/Source/MeshPaintingToolset/Public/MeshPaintHelpers.h
+++ b/Engine/Plugins/MeshPainting/Source/MeshPaintingToolset/Public/MeshPaintHelpers.h
@@ -5,6 +5,7 @@
 #include "Engine/Engine.h"
 #include "Engine/HitResult.h"
 #include "Engine/StaticMesh.h"
+#include "ImageCore.h"
 #include "Math/Ray.h"
 #include "UObject/Package.h"
 #include "MeshPaintHelpers.generated.h"
@@ -28,6 +29,7 @@ class FSceneView;
 struct FStaticMeshComponentLODInfo;
 class UMeshVertexPaintingToolProperties;
 class UBrushBaseProperties;
+struct FMeshDescription;
 
 enum class EMeshPaintDataColorViewMode : uint8;
 
@@ -134,7 +136,14 @@ public:
 
 	/** Removes vertex colors associated with the mesh component */
 	void RemoveComponentInstanceVertexColors(UStaticMeshComponent* StaticMeshComponent);
-	
+
+	/** Creates mesh paint texture associated with the mesh component */
+	void CreateComponentMeshPaintTexture(UStaticMeshComponent* StaticMeshComponent);
+	void CreateComponentMeshPaintTexture(UStaticMeshComponent* StaticMeshComponent, FImageView const& InImage);
+
+	/** Removes mesh paint texture associated with the mesh component */
+	void RemoveComponentMeshPaintTexture(UStaticMeshComponent* StaticMeshComponent);
+
 	/** Propagates per-instance vertex colors to the underlying Mesh for the given LOD Index */
 	bool PropagateColorsToRawMesh(UStaticMesh* StaticMesh, int32 LODIndex, FStaticMeshComponentLODInfo& ComponentLODInfo);	
 
@@ -260,6 +269,9 @@ public:
 	bool SelectionContainsPerLODColors() const { return bSelectionContainsPerLODColors; }
 	void ClearSelectionLODColors() { bSelectionContainsPerLODColors = false; }
 
+	FImage const& GetCopiedTexture() const;
+	void SetCopiedTexture(UTexture* InTexture);
+
 public:
 	bool bNeedsRecache;
 	UPROPERTY(Transient)
@@ -285,6 +297,9 @@ private:
 	/** Contains copied vertex color data */
 	TArray<FPerComponentVertexColorData> CopiedColorsByComponent;
 	bool bSelectionContainsPerLODColors;
+
+	/** Contains copied texture data */
+	FImage CopiedTextureData;
 };
 
 template<typename T>
diff --git a/Engine/Plugins/MeshPainting/Source/MeshPaintingToolset/Public/MeshPaintingToolsetTypes.h b/Engine/Plugins/MeshPainting/Source/MeshPaintingToolset/Public/MeshPaintingToolsetTypes.h
--- a/Engine/Plugins/MeshPainting/Source/MeshPaintingToolset/Public/MeshPaintingToolsetTypes.h
+++ b/Engine/Plugins/MeshPainting/Source/MeshPaintingToolset/Public/MeshPaintingToolsetTypes.h
@@ -84,7 +84,6 @@ public:
 	bool bWriteAlpha;
 	int32 TotalWeightCount;
 	int32 PaintWeightIndex;
-	int32 UVChannel;
 	FApplyVertexPaintData ApplyVertexDataDelegate;
 	FVector2f BrushPosition2D;
 	bool bUseFillBucket = false;
@@ -103,19 +102,22 @@ struct FPaintableTexture
 {
 	UTexture*	Texture;
 	int32		UVChannelIndex;
+	bool		bIsMeshTexture;
 
 	FPaintableTexture()
 		: Texture(nullptr)
 		, UVChannelIndex(0)
+		, bIsMeshTexture(false)
 	{}
 
 	template <
 		typename T
 		UE_REQUIRES(std::is_convertible_v<T, UTexture*>)
 	>
-	FPaintableTexture(T InTexture = nullptr, uint32 InUVChannelIndex = 0)
+	FPaintableTexture(T InTexture = nullptr, uint32 InUVChannelIndex = 0, bool bInIsMeshTexture = false)
 		: Texture(InTexture)
 		, UVChannelIndex(InUVChannelIndex)
+		, bIsMeshTexture(bInIsMeshTexture)
 	{}
 
 	/** Overloaded equality operator for use with TArrays Contains method. */
@@ -130,6 +132,7 @@ USTRUCT()
 struct FPaintTexture2DData
 {
 	GENERATED_BODY()
+
 	/** The original texture that we're painting */
 	UPROPERTY(Transient)
 	TObjectPtr<UTexture2D> PaintingTexture2D = nullptr;
@@ -157,6 +160,22 @@ struct FPaintTexture2DData
 	UPROPERTY(Transient)
 	TObjectPtr<class UVirtualTextureAdapter> PaintRenderTargetTextureAdapter = nullptr;
 
+	/** Temporary render target used to draw incremental paint to */
+	UPROPERTY(Transient)
+	TObjectPtr<UTextureRenderTarget2D> BrushRenderTargetTexture;
+
+	/** Temporary render target used to store a mask of the affected paint region, updated every time we add incremental texture paint */
+	UPROPERTY(Transient)
+	TObjectPtr<UTextureRenderTarget2D> BrushMaskRenderTargetTexture;
+
+	/** Optional render target used to store generated mask for texture seams. We create this by projecting object triangles into texture space using the selected UV channel. */
+	UPROPERTY(Transient)
+	TObjectPtr<UTextureRenderTarget2D> SeamMaskRenderTargetTexture = nullptr;
+
+	/** True if we need to generate a texture seam mask used for texture dilation */
+	UPROPERTY(Transient)
+	bool bGenerateSeamMask;
+
 	FPaintTexture2DData() = default;
 
 	FPaintTexture2DData(UTexture2D* InPaintingTexture2D, bool InbIsPaintingTexture2DModified = false)
@@ -164,7 +183,6 @@ struct FPaintTexture2DData
 		, bIsPaintingTexture2DModified(InbIsPaintingTexture2DModified)
 	{
 	}
-
 };
 
 USTRUCT()
diff --git a/Engine/Plugins/MeshPainting/Source/MeshPaintingToolset/Public/MeshTexturePaintingTool.h b/Engine/Plugins/MeshPainting/Source/MeshPaintingToolset/Public/MeshTexturePaintingTool.h
--- a/Engine/Plugins/MeshPainting/Source/MeshPaintingToolset/Public/MeshTexturePaintingTool.h
+++ b/Engine/Plugins/MeshPainting/Source/MeshPaintingToolset/Public/MeshTexturePaintingTool.h
@@ -3,6 +3,7 @@
 #pragma once
 
 #include "BaseTools/BaseBrushTool.h"
+#include "BaseMeshPaintingToolProperties.h"
 #include "MeshPaintHelpers.h"
 #include "MeshPaintInteractions.h"
 #include "MeshPaintingToolsetTypes.h"
@@ -11,25 +12,23 @@
 
 #include "MeshTexturePaintingTool.generated.h"
 
-class UMeshToolManager;
+enum class EMeshPaintModeAction : uint8;
 enum class EToolShutdownType : uint8;
-struct FTexturePaintMeshSectionInfo;
-
-
-struct FToolBuilderState;
-struct FPaintTexture2DData;
+class FScopedTransaction;
+class IMeshPaintComponentAdapter;
+class UMeshToolManager;
 class UTexture2D;
 class UTextureRenderTarget2D;
+struct FPaintTexture2DData;
+struct FTexturePaintMeshSectionInfo;
 struct FTextureTargetListInfo;
-enum class EMeshPaintModeAction : uint8;
-class IMeshPaintComponentAdapter;
-class FScopedTransaction;
+struct FToolBuilderState;
 
 /**
- *
+ * Builder for the texture color mesh paint tool.
  */
 UCLASS()
-class MESHPAINTINGTOOLSET_API UMeshTexturePaintingToolBuilder : public UInteractiveToolBuilder
+class MESHPAINTINGTOOLSET_API UMeshTextureColorPaintingToolBuilder : public UInteractiveToolBuilder
 {
 	GENERATED_BODY()
 
@@ -40,50 +39,34 @@
 	TWeakObjectPtr<UMeshToolManager> SharedMeshToolData;
 };
 
-
+/**
+ * Builder for the texture asset mesh paint tool.
+ */
 UCLASS()
-class MESHPAINTINGTOOLSET_API UMeshTexturePaintingToolProperties : public UBrushBaseProperties
+class MESHPAINTINGTOOLSET_API UMeshTextureAssetPaintingToolBuilder : public UInteractiveToolBuilder
 {
 	GENERATED_BODY()
 
 public:
-	UMeshTexturePaintingToolProperties();
+	virtual bool CanBuildTool(const FToolBuilderState& SceneState) const override;
+	virtual UInteractiveTool* BuildTool(const FToolBuilderState& SceneState) const override;
 
-	/** Color used for Applying Texture Color Painting */
-	UPROPERTY(EditAnywhere, Category = TexturePainting)
-	FLinearColor PaintColor;
+	TWeakObjectPtr<UMeshToolManager> SharedMeshToolData;
+};
 
-	/** Color used for Erasing Texture Color Painting */
-	UPROPERTY(EditAnywhere, Category = TexturePainting)
-	FLinearColor EraseColor;
 
-	/** Whether or not to apply Texture Color Painting to the Red Channel */
-	UPROPERTY(EditAnywhere, Category = TexturePainting, meta = (DisplayName = "Red"))
-	bool bWriteRed;
+/**
+ * Base class for mesh texture paint properties.
+ */
+UCLASS()
+class MESHPAINTINGTOOLSET_API UMeshTexturePaintingToolProperties : public UMeshPaintingToolProperties
+{
+	GENERATED_BODY()
 
-	/** Whether or not to apply Texture Color Painting to the Green Channel */
-	UPROPERTY(EditAnywhere, Category = TexturePainting, meta = (DisplayName = "Green"))
-	bool bWriteGreen;
-
-	/** Whether or not to apply Texture Color Painting to the Blue Channel */
-	UPROPERTY(EditAnywhere, Category = TexturePainting, meta = (DisplayName = "Blue"))
-	bool bWriteBlue;
-
-	/** Whether or not to apply Texture Color Painting to the Alpha Channel */
-	UPROPERTY(EditAnywhere, Category = TexturePainting, meta = (DisplayName = "Alpha"))
-	bool bWriteAlpha;
-
-	/** UV channel which should be used for paint textures */
-	UPROPERTY(EditAnywhere, Category = TexturePainting, meta = (TransientToolProperty))
-	int32 UVChannel;
-
+public:
 	/** Seam painting flag, True if we should enable dilation to allow the painting of texture seams */
 	UPROPERTY(EditAnywhere, Category = TexturePainting)
-	bool bEnableSeamPainting;
+	bool bEnableSeamPainting = false;
 
-	/** Texture to which Painting should be Applied */
-	UPROPERTY(EditAnywhere, Category = TexturePainting, meta = (DisplayThumbnail = "true", TransientToolProperty))
-	TObjectPtr<UTexture2D> PaintTexture;
-
 	/** Optional Texture Brush to which Painting should use */
 	UPROPERTY(EditAnywhere, Category = TexturePainting, meta = (DisplayThumbnail = "true", TransientToolProperty))
@@ -97,18 +80,55 @@
 	UPROPERTY(EditAnywhere, Category = TexturePainting, meta = (TransientToolProperty))
 	bool bRotateBrushTowardsDirection = false;
 
-	/** Enables "Flow" painting where paint is continually applied from the brush every tick */
-	UPROPERTY(EditAnywhere, Category = Brush, meta = (DisplayName = "Enable Brush Flow"))
-	bool bEnableFlow;
+	/** Whether or not to apply Texture Color Painting to the Red Channel */
+	UPROPERTY(EditAnywhere, Category = ColorPainting, meta = (DisplayName = "Red"))
+	bool bWriteRed = true;
 
-	/** Whether back-facing triangles should be ignored */
-	UPROPERTY(EditAnywhere, Category = Brush, meta = (DisplayName = "Ignore Back-Facing"))
-	bool bOnlyFrontFacingTriangles;
+	/** Whether or not to apply Texture Color Painting to the Green Channel */
+	UPROPERTY(EditAnywhere, Category = ColorPainting, meta = (DisplayName = "Green"))
+	bool bWriteGreen = true;
+
+	/** Whether or not to apply Texture Color Painting to the Blue Channel */
+	UPROPERTY(EditAnywhere, Category = ColorPainting, meta = (DisplayName = "Blue"))
+	bool bWriteBlue = true;
+
+	/** Whether or not to apply Texture Color Painting to the Alpha Channel */
+	UPROPERTY(EditAnywhere, Category = ColorPainting, meta = (DisplayName = "Alpha"))
+	bool bWriteAlpha = false;
+};
+
+/**
+ * Class for texture color paint properties.
+ */
+UCLASS()
+class MESHPAINTINGTOOLSET_API UMeshTextureColorPaintingToolProperties : public UMeshTexturePaintingToolProperties
+{
+	GENERATED_BODY()
+};
+
+/**
+ * Class for texture asset paint properties.
+ */
+UCLASS()
+class MESHPAINTINGTOOLSET_API UMeshTextureAssetPaintingToolProperties : public UMeshTexturePaintingToolProperties
+{
+	GENERATED_BODY()
+
+public:
+	/** UV channel which should be used for painting textures. */
+	UPROPERTY(EditAnywhere, Category = TexturePainting, meta = (TransientToolProperty))
+	int32 UVChannel = 0;
+
+	/** Texture to which painting should be applied. */
+	UPROPERTY(EditAnywhere, Category = TexturePainting, meta = (DisplayThumbnail = "true", TransientToolProperty))
+	TObjectPtr<UTexture2D> PaintTexture;
 };
 
 
-
-UCLASS()
+/**
+ * Base class for mesh texture painting tool.
+ */
+UCLASS(Abstract)
 class MESHPAINTINGTOOLSET_API UMeshTexturePaintingTool : public UBaseBrushTool, public IMeshPaintSelectionInterface
 {
 	GENERATED_BODY()
@@ -136,26 +156,22 @@
 		return bArePainting;
 	}
 
-
 	FSimpleDelegate& OnPaintingFinished()
 	{
 		return OnPaintingFinishedDelegate;
 	}
-	void CycleTextures(int32 Direction);
 
 	void CommitAllPaintedTextures();
 	void ClearAllTextureOverrides();
+	
 	/** Returns the number of texture that require a commit. */
 	int32 GetNumberOfPendingPaintChanges() const;
 
 	void FloodCurrentPaintTexture();
-	bool ShouldFilterTextureAsset(const FAssetData& AssetData) const;
-	void PaintTextureChanged(const FAssetData& AssetData);
+	
 	virtual bool IsMeshAdapterSupported(TSharedPtr<IMeshPaintComponentAdapter> MeshAdapter) const override;
-	virtual bool AllowsMultiselect() const override
-	{
-		return false;
-	}
+	
+	virtual void GetModifiedTexturesToSave(TArray<UObject*>& OutTexturesToSave) const {}
 
 protected:
 	virtual void SetAdditionalPaintParameters(FMeshPaintParameters& InPaintParameters) {};
@@ -165,69 +181,35 @@
 	bool Paint(const FVector& InRayOrigin, const FVector& InRayDirection);
 	bool Paint(const TArrayView<TPair<FVector, FVector>>& Rays);
 	virtual void CacheSelectionData();
-	void CacheTexturePaintData();
 	FPaintTexture2DData* GetPaintTargetData(const UTexture2D* InTexture);
 	FPaintTexture2DData* AddPaintTargetData(UTexture2D* InTexture);
 	void GatherTextureTriangles(IMeshPaintComponentAdapter* Adapter, int32 TriangleIndex, const int32 VertexIndices[3], TArray<FTexturePaintTriangleInfo>* TriangleInfo, TArray<FTexturePaintMeshSectionInfo>* SectionInfos, int32 UVChannelIndex);
 	void StartPaintingTexture(UMeshComponent* InMeshComponent, const IMeshPaintComponentAdapter& GeometryInfo);
-	void PaintTexture(FMeshPaintParameters& InParams, TArray<FTexturePaintTriangleInfo>& InInfluencedTriangles, const IMeshPaintComponentAdapter& GeometryInfo, FMeshPaintParameters* LastParams = nullptr);
+	void PaintTexture(FMeshPaintParameters& InParams, int32 UVChannel, TArray<FTexturePaintTriangleInfo>& InInfluencedTriangles, const IMeshPaintComponentAdapter& GeometryInfo, FMeshPaintParameters* LastParams = nullptr);
 	void FinishPaintingTexture();
 	void OnTransactionStateChanged(const FTransactionContext& InTransactionContext, const ETransactionStateEventType InTransactionState);
 
-protected:
-	double InitialMeshArea;
-	bool bResultValid;
-	bool bStampPending;
-	bool bInDrag;
-	FRay PendingStampRay;
-	FRay PendingClickRay;
-	FVector2D PendingClickScreenPosition;
-	bool bCachedClickRay;
-
-	UPROPERTY(Transient)
-	TObjectPtr<UMeshPaintSelectionMechanic> SelectionMechanic;
+	virtual UTexture2D* GetSelectedPaintTexture(UMeshComponent const* InMeshComponent) const { return nullptr; }
+	virtual int32 GetSelectedUVChannel(UMeshComponent const* InMeshComponent) const { return 0; }
+	virtual void CacheTexturePaintData() {}
 
 private:
 	bool PaintInternal(const TArrayView<TPair<FVector, FVector>>& Rays, EMeshPaintModeAction PaintAction, float PaintStrength);
 
 	void AddTextureOverrideToComponent(FPaintTexture2DData& TextureData, UMeshComponent* MeshComponent, const IMeshPaintComponentAdapter* MeshPaintAdapter = nullptr);
 
-private:
+protected:
 	UPROPERTY(Transient)
+	TObjectPtr<UMeshPaintSelectionMechanic> SelectionMechanic;
+
+	/** Textures eligible for painting retrieved from the current selection */
+	TArray<FPaintableTexture> PaintableTextures;
+
+	UPROPERTY(Transient)
 	TObjectPtr<UMeshTexturePaintingToolProperties> TextureProperties;
 
 	UPROPERTY(Transient)
 	TArray<TObjectPtr<const UTexture>> Textures;
 
-	TArray<FPaintRayResults> LastPaintRayResults;
-	bool bRequestPaintBucketFill = false;
-
-	/** Flag for whether or not we are currently painting */
-	bool bArePainting;
-	bool bDoRestoreRenTargets;
-	/** Time kept since the user has started painting */
-	float TimeSinceStartedPainting;
-	/** Overall time value kept for drawing effects */
-	float Time;
-	FHitResult LastBestHitResult;
-	FSimpleDelegate OnPaintingFinishedDelegate;
-	/** Texture paint state */
-/** Textures eligible for painting retrieved from the current selection */
-	TArray<FPaintableTexture> PaintableTextures;
-	/** Cached / stored instance texture paint settings for selected components */
-	TMap<UMeshComponent*, FInstanceTexturePaintSettings> ComponentToTexturePaintSettingsMap;
-
-	/** Temporary render target used to draw incremental paint to */
-	UPROPERTY(Transient)
-	TObjectPtr<UTextureRenderTarget2D> BrushRenderTargetTexture;
-
-	/** Temporary render target used to store a mask of the affected paint region, updated every time we add incremental texture paint */
-	UPROPERTY(Transient)
-	TObjectPtr<UTextureRenderTarget2D> BrushMaskRenderTargetTexture;
-
-	/** Temporary render target used to store generated mask for texture seams, we create this by projecting object triangles into texture space using the selected UV channel */
-	UPROPERTY(Transient)
-	TObjectPtr<UTextureRenderTarget2D> SeamMaskRenderTargetTexture;
-
 	/** Stores data associated with our paint target textures */
 	UPROPERTY(Transient)
@@ -250,11 +232,101 @@
 	UPROPERTY(Transient)
 	TObjectPtr<UTexture2D> PaintingTexture2D;
 
-	/** True if we need to generate a texture seam mask used for texture dilation */
-	bool bGenerateSeamMask;
-
 	/** Hold the transaction while we are painting */
 	TUniquePtr<FScopedTransaction> PaintingTransaction;
+
+	double InitialMeshArea;
+	bool bResultValid;
+	bool bStampPending;
+	bool bInDrag;
+	FRay PendingStampRay;
+	FRay PendingClickRay;
+	FVector2D PendingClickScreenPosition;
+	bool bCachedClickRay;
+
+	TArray<FPaintRayResults> LastPaintRayResults;
+	bool bRequestPaintBucketFill = false;
+
+	/** Flag for whether or not we are currently painting */
+	bool bArePainting;
+	bool bDoRestoreRenTargets;
+	/** Time kept since the user has started painting */
+	float TimeSinceStartedPainting;
+	/** Overall time value kept for drawing effects */
+	float Time;
+	FHitResult LastBestHitResult;
+	FSimpleDelegate OnPaintingFinishedDelegate;
+	/** Texture paint state */
+	/** Cached / stored instance texture paint settings for selected components */
+	TMap<UMeshComponent*, FInstanceTexturePaintSettings> ComponentToTexturePaintSettingsMap;
 };
 
+/**
+ * Class for texture color painting tool.
+ * This paints to special textures stored on the mesh components.
+ * Behavior should be similar to vertex painting (per instance painting stored on components).
+ * But painting texture colors instead of vertex colors is a better fit for very dense mesh types such as used by nanite.
+ */
+UCLASS()
+class MESHPAINTINGTOOLSET_API UMeshTextureColorPaintingTool : public UMeshTexturePaintingTool
+{
+	GENERATED_BODY()
+
+public:
+	UMeshTextureColorPaintingTool();
+
+	// Begin UInteractiveTool Interface.
+	virtual void Setup() override;
+	// End UInteractiveTool Interface.
+
+	// Begin UMeshTexturePaintingTool Interface.
+	virtual bool AllowsMultiselect() const override { return true; }
+	virtual UTexture2D* GetSelectedPaintTexture(UMeshComponent const* InMeshComponent) const override;
+	virtual int32 GetSelectedUVChannel(UMeshComponent const* InMeshComponent) const override;
+	virtual void GetModifiedTexturesToSave(TArray<UObject*>& OutTexturesToSave) const override;
+	virtual void CacheTexturePaintData() override;
+	// End UMeshTexturePaintingTool Interface.
+
+protected:
+	UPROPERTY(Transient)
+	TObjectPtr<UMeshTextureColorPaintingToolProperties> ColorProperties;
+};
+
+/**
+ * Class for texture asset painting tool.
+ * This paints to texture assets directly from the mesh.
+ * The texture asset to paint is selected from the ones referenced in the mesh component's materials.
+ */
+UCLASS()
+class MESHPAINTINGTOOLSET_API UMeshTextureAssetPaintingTool : public UMeshTexturePaintingTool
+{
+	GENERATED_BODY()
+
+public:
+	UMeshTextureAssetPaintingTool();
+	
+	// Begin UInteractiveTool Interface.
+	virtual void Setup() override;
+	// End UInteractiveTool Interface.
+
+	// Begin UMeshTexturePaintingTool Interface.
+	virtual bool AllowsMultiselect() const override { return false; }
+	virtual UTexture2D* GetSelectedPaintTexture(UMeshComponent const* InMeshComponent) const override;
+	virtual int32 GetSelectedUVChannel(UMeshComponent const* InMeshComponent) const override;
+	virtual void GetModifiedTexturesToSave(TArray<UObject*>& OutTexturesToSave) const override;
+	virtual void CacheTexturePaintData() override;
+	// End UMeshTexturePaintingTool Interface.
+
+	/** Change selected texture to previous or next available. */
+	void CycleTextures(int32 Direction);
+	/** Returns true if asset shouldn't be shown in UI because it is not in our paintable texture array. */
+	bool ShouldFilterTextureAsset(const FAssetData& AssetData) const;
+	/** Call on change to selected paint texture */
+	void PaintTextureChanged(const FAssetData& AssetData);
+
+protected:
+	UPROPERTY(Transient)
+	TObjectPtr<UMeshTextureAssetPaintingToolProperties> AssetProperties;
+};
+
 #if UE_ENABLE_INCLUDE_ORDER_DEPRECATED_IN_5_2
diff --git a/Engine/Plugins/MeshPainting/Source/MeshPaintingToolset/Public/MeshVertexPaintingTool.h b/Engine/Plugins/MeshPainting/Source/MeshPaintingToolset/Public/MeshVertexPaintingTool.h
--- a/Engine/Plugins/MeshPainting/Source/MeshPaintingToolset/Public/MeshVertexPaintingTool.h
+++ b/Engine/Plugins/MeshPainting/Source/MeshPaintingToolset/Public/MeshVertexPaintingTool.h
@@ -3,6 +3,7 @@
 #pragma once
 
 #include "BaseTools/BaseBrushTool.h"
+#include "BaseMeshPaintingToolProperties.h"
 #include "MeshPaintingToolsetTypes.h"
 #include "MeshPaintInteractions.h"
 #include "MeshVertexPaintingTool.generated.h"
@@ -74,29 +75,13 @@ public:
 
 
 UCLASS()
-class MESHPAINTINGTOOLSET_API UMeshVertexPaintingToolProperties : public UBrushBaseProperties
+class MESHPAINTINGTOOLSET_API UMeshVertexPaintingToolProperties : public UMeshPaintingToolProperties
 {
 	GENERATED_BODY()
 
 public:
 	UMeshVertexPaintingToolProperties();
 
-	/** Color used for Applying Vertex Color Painting */
-	UPROPERTY(EditAnywhere, Category = VertexPainting)
-	FLinearColor PaintColor;
-
-	/** Color used for Erasing Vertex Color Painting */
-	UPROPERTY(EditAnywhere, Category = VertexPainting)
-	FLinearColor EraseColor;
-
-	/** Enables "Flow" painting where paint is continually applied from the brush every tick */
-	UPROPERTY(EditAnywhere, Category = Brush, meta = (DisplayName = "Enable Brush Flow"))
-	bool bEnableFlow;
-
-	/** Whether back-facing triangles should be ignored */
-	UPROPERTY(EditAnywhere, Category = Brush, meta = (DisplayName = "Ignore Back-Facing"))
-	bool bOnlyFrontFacingTriangles;
-
 	/** Size of vertex points drawn when mesh painting is active. */
 	UPROPERTY(EditAnywhere, Category = "VertexPainting|Visualization")
 	float VertexPreviewSize;
@@ -104,46 +89,43 @@ public:
 
 
 UCLASS()
-class MESHPAINTINGTOOLSET_API UMeshColorPaintingToolProperties : public UMeshVertexPaintingToolProperties
+class MESHPAINTINGTOOLSET_API UMeshVertexColorPaintingToolProperties : public UMeshVertexPaintingToolProperties
 {
 	GENERATED_BODY()
 
 public:
-	UMeshColorPaintingToolProperties();
-
-
 	/** Whether or not to apply Vertex Color Painting to the Red Channel */
 	UPROPERTY(EditAnywhere, Category = ColorPainting, DisplayName = "Red")
-	bool bWriteRed;
+	bool bWriteRed = true;
 
 	/** Whether or not to apply Vertex Color Painting to the Green Channel */
 	UPROPERTY(EditAnywhere, Category = ColorPainting, DisplayName = "Green")
-	bool bWriteGreen;
+	bool bWriteGreen = true;
 
 	/** Whether or not to apply Vertex Color Painting to the Blue Channel */
 	UPROPERTY(EditAnywhere, Category = ColorPainting, DisplayName = "Blue")
-	bool bWriteBlue;
+	bool bWriteBlue = true;
 
 	/** Whether or not to apply Vertex Color Painting to the Alpha Channel */
 	UPROPERTY(EditAnywhere, Category = ColorPainting, DisplayName = "Alpha")
-	bool bWriteAlpha;
+	bool bWriteAlpha = false;
 
 	/** When unchecked the painting on the base LOD will be propagate automatically to all other LODs when exiting the mode or changing the selection */
 	UPROPERTY(EditAnywhere, Category = Painting, meta = (TransientToolProperty))
-	bool bPaintOnSpecificLOD;
+	bool bPaintOnSpecificLOD = false;
 
 	/** LOD Index to which should specifically be painted */
 	UPROPERTY(EditAnywhere, Category = Painting, meta = (UIMin = "0", ClampMin = "0", EditCondition = "bPaintOnSpecificLOD", TransientToolProperty))
-	int32 LODIndex;
+	int32 LODIndex = 0;
 };
 
 UCLASS()
-class MESHPAINTINGTOOLSET_API UMeshWeightPaintingToolProperties : public UMeshVertexPaintingToolProperties
+class MESHPAINTINGTOOLSET_API UMeshVertexWeightPaintingToolProperties : public UMeshVertexPaintingToolProperties
 {
 	GENERATED_BODY()
 
 public:
-	UMeshWeightPaintingToolProperties();
+	UMeshVertexWeightPaintingToolProperties();
 
 	/** Texture Blend Weight Painting Mode */
 	UPROPERTY(EditAnywhere, Category = WeightPainting, meta = (EnumCondition = 1))
@@ -264,7 +246,7 @@ protected:
 
 private:
 	UPROPERTY(Transient)
-	TObjectPtr<UMeshColorPaintingToolProperties> ColorProperties;
+	TObjectPtr<UMeshVertexColorPaintingToolProperties> ColorProperties;
 
 	/** Current LOD index used for painting / forcing */
 	int32 CachedLODIndex;
@@ -288,7 +270,7 @@ protected:
 
 private:
 	UPROPERTY(Transient)
-	TObjectPtr<UMeshWeightPaintingToolProperties> WeightProperties;
+	TObjectPtr<UMeshVertexWeightPaintingToolProperties> WeightProperties;
 };
 
 #if UE_ENABLE_INCLUDE_ORDER_DEPRECATED_IN_5_2
diff --git a/Engine/Plugins/MeshPainting/Source/MeshPaintingToolset/Public/TexturePaintToolset.h b/Engine/Plugins/MeshPainting/Source/MeshPaintingToolset/Public/TexturePaintToolset.h
--- a/Engine/Plugins/MeshPainting/Source/MeshPaintingToolset/Public/TexturePaintToolset.h
+++ b/Engine/Plugins/MeshPainting/Source/MeshPaintingToolset/Public/TexturePaintToolset.h
@@ -101,7 +101,7 @@ public:
 	static void RetrieveMeshSectionsForMaterialIndices(const UMeshComponent* MeshComponent, int32 LODIndex, const TArray<int32>& MaterialIndices, TArray<FTexturePaintMeshSectionInfo>& OutSectionInfo);
 
 	/** Retrieves all Paintable Textures from the given MeshComponent */
-	static void RetrieveTexturesForComponent(const UMeshComponent* Component, IMeshPaintComponentAdapter* Adapter, TArray<FPaintableTexture>& OutTextures);
+	static void RetrieveTexturesForComponent(const UMeshComponent* Component, IMeshPaintComponentAdapter* Adapter, int32& OutDefaultIndex, TArray<FPaintableTexture>& OutTextures);
 
 	/** Checks whether or not the given MeshComponent's materials use the specified texture */
 	static bool DoesMeshComponentUseTexture(UMeshComponent* MeshComponent, UTexture* Texture);
diff --git a/Engine/Source/Editor/EditorStyle/Private/StarshipStyle.cpp b/Engine/Source/Editor/EditorStyle/Private/StarshipStyle.cpp
--- a/Engine/Source/Editor/EditorStyle/Private/StarshipStyle.cpp
+++ b/Engine/Source/Editor/EditorStyle/Private/StarshipStyle.cpp
@@ -4838,54 +4838,63 @@ void FStarshipEditorStyle::FStyle::SetupLevelEditorStyle()
 
 	// Mesh Paint
 	{
-		Set("MeshPaint.FindInCB",       new IMAGE_BRUSH("/Icons/icon_toolbar_genericfinder_40px", Icon20x20) );
-		Set("MeshPaint.Swap",           new IMAGE_BRUSH("/Icons/icon_MeshPaint_Swap_16x", Icon12x12));
-
-		Set("MeshPaint.Fill",           new IMAGE_BRUSH("/Icons/GeneralTools/PaintBucket_40x", Icon20x20) );
-		Set("MeshPaint.Fill.Small",     new IMAGE_BRUSH("/Icons/GeneralTools/PaintBucket_40x", Icon20x20) );
-		Set("MeshPaint.Propagate",      new IMAGE_BRUSH("/Icons/Paint/Paint_Propagate_40x", Icon20x20) );
-		Set("MeshPaint.Propagate.Small",new IMAGE_BRUSH("/Icons/Paint/Paint_Propagate_40x", Icon20x20) );
-		Set("MeshPaint.Import",         new IMAGE_BRUSH("/Icons/GeneralTools/Import_40x", Icon20x20) );
-		Set("MeshPaint.Import.Small",   new IMAGE_BRUSH("/Icons/GeneralTools/Import_40x", Icon20x20) );
-		Set("MeshPaint.Save",           new IMAGE_BRUSH("/Icons/GeneralTools/Save_40x", Icon20x20) );
-		Set("MeshPaint.Save.Small",     new IMAGE_BRUSH("/Icons/GeneralTools/Save_40x", Icon20x20) );
-		Set("MeshPaint.Fix",            new IMAGE_BRUSH("/Icons/GeneralTools/Fix_40x", Icon20x20) );
-		Set("MeshPaint.Fix.Small",      new IMAGE_BRUSH("/Icons/GeneralTools/Fix_40x", Icon20x20) );
-		Set("MeshPaint.Remove",         new IMAGE_BRUSH("/Icons/GeneralTools/Delete_40x", Icon20x20));
-		Set("MeshPaint.Remove.Small",   new IMAGE_BRUSH("/Icons/GeneralTools/Delete_40x", Icon20x20));
-		Set("MeshPaint.Copy",           new IMAGE_BRUSH("/Icons/GeneralTools/Copy_40x", Icon20x20));
-		Set("MeshPaint.Copy.Small",     new IMAGE_BRUSH("/Icons/GeneralTools/Copy_40x", Icon20x20));
-		Set("MeshPaint.Paste",          new IMAGE_BRUSH("/Icons/GeneralTools/Paste_40x", Icon20x20));
-		Set("MeshPaint.Paste.Small",    new IMAGE_BRUSH("/Icons/GeneralTools/Paste_40x", Icon20x20));
-		Set("MeshPaint.PropagateTexturePaint",             new IMAGE_BRUSH("/Icons/Paint/Paint_PropagateTexture_40x", Icon20x20));
-		Set("MeshPaint.PropagateTexturePaint.Small",       new IMAGE_BRUSH("/Icons/Paint/Paint_PropagateTexture_40x", Icon20x20));
-		Set("MeshPaint.SaveTexturePaint",                  new IMAGE_BRUSH("/Icons/GeneralTools/Save_40x", Icon20x20));
-		Set("MeshPaint.SaveTexturePaint.Small",            new IMAGE_BRUSH("/Icons/GeneralTools/Save_40x", Icon20x20));
-		Set("MeshPaint.PropagateVertexColorsToLODs",       new IMAGE_BRUSH("Icons/Paint/Paint_AllLODs_40x", Icon20x20));
-		Set("MeshPaint.PropagateVertexColorsToLODs.Small", new IMAGE_BRUSH("Icons/Paint/Paint_AllLODs_40x", Icon20x20));
-		Set("MeshPaint.TexturePaint",       new IMAGE_BRUSH("Icons/GeneralTools/Paint_40x", Icon20x20));
-		Set("MeshPaint.TexturePaint.Small", new IMAGE_BRUSH("Icons/GeneralTools/Paint_40x", Icon20x20));
-		Set("MeshPaint.ColorPaint",         new IMAGE_BRUSH("Icons/GeneralTools/Paint_40x", Icon20x20));
-		Set("MeshPaint.ColorPaint.Small",   new IMAGE_BRUSH("Icons/GeneralTools/Paint_40x", Icon20x20));
-		Set("MeshPaint.WeightPaint",        new IMAGE_BRUSH("Icons/GeneralTools/Paint_40x", Icon20x20));
-		Set("MeshPaint.WeightPaint.Small",  new IMAGE_BRUSH("Icons/GeneralTools/Paint_40x", Icon20x20));
-		Set("MeshPaint.VertexSelect",             new IMAGE_BRUSH(TEXT("Icons/GeneralTools/Select_40x"), Icon20x20));
-		Set("MeshPaint.VertexSelect.Small",       new IMAGE_BRUSH(TEXT("Icons/GeneralTools/Select_40x"), Icon20x20));
-		Set("MeshPaint.TextureSelect",            new IMAGE_BRUSH(TEXT("Icons/GeneralTools/Select_40x"), Icon20x20));
-		Set("MeshPaint.TextureSelect.Small",      new IMAGE_BRUSH(TEXT("Icons/GeneralTools/Select_40x"), Icon20x20));
-		Set("MeshPaint.PreviousTexture",          new IMAGE_BRUSH(TEXT("Icons/GeneralTools/Previous_40x"), Icon20x20));
-		Set("MeshPaint.PreviousTexture.Small",    new IMAGE_BRUSH(TEXT("Icons/GeneralTools/Previous_40x"), Icon20x20));
-		Set("MeshPaint.NextTexture",              new IMAGE_BRUSH(TEXT("Icons/GeneralTools/Next_40x"), Icon20x20));
-		Set("MeshPaint.NextTexture.Small",        new IMAGE_BRUSH(TEXT("Icons/GeneralTools/Next_40x"), Icon20x20));
-		Set("MeshPaint.CycleToPreviousLOD",       new IMAGE_BRUSH(TEXT("Icons/GeneralTools/Previous_40x"), Icon20x20));
-		Set("MeshPaint.CycleToPreviousLOD.Small", new IMAGE_BRUSH(TEXT("Icons/GeneralTools/Previous_40x"), Icon20x20));
-		Set("MeshPaint.CycleToNextLOD",           new IMAGE_BRUSH(TEXT("Icons/GeneralTools/Next_40x"), Icon20x20));
-		Set("MeshPaint.CycleToNextLOD.Small",     new IMAGE_BRUSH(TEXT("Icons/GeneralTools/Next_40x"), Icon20x20));
-		Set("MeshPaint.SwitchForeAndBackgroundColor",       new IMAGE_BRUSH("Icons/Paint/Paint_SwapColors_40x", Icon20x20));
-		Set("MeshPaint.SwitchForeAndBackgroundColor.Small", new IMAGE_BRUSH("Icons/Paint/Paint_SwapColors_40x", Icon20x20));
-		Set("MeshPaint.Brush",			new IMAGE_BRUSH_SVG("Starship/Common/Paintbrush", Icon20x20));
-		Set("MeshPaint.TextureFill", new IMAGE_BRUSH("/Icons/GeneralTools/PaintBucket_40x", Icon20x20));
-		Set("MeshPaint.TextureFill.Small", new IMAGE_BRUSH("/Icons/GeneralTools/PaintBucket_40x", Icon20x20));
+		Set("MeshPaint.SelectVertex", new IMAGE_BRUSH(TEXT("Icons/GeneralTools/Select_40x"), Icon20x20));
+		Set("MeshPaint.SelectVertex.Small", new IMAGE_BRUSH(TEXT("Icons/GeneralTools/Select_40x"), Icon20x20));
+		Set("MeshPaint.SelectTexture", new IMAGE_BRUSH(TEXT("Icons/GeneralTools/Select_40x"), Icon20x20));
+		Set("MeshPaint.SelectTexture.Small", new IMAGE_BRUSH(TEXT("Icons/GeneralTools/Select_40x"), Icon20x20));
+		Set("MeshPaint.PaintVertexColor", new IMAGE_BRUSH("Icons/GeneralTools/Paint_40x", Icon20x20));
+		Set("MeshPaint.PaintVertexColor.Small", new IMAGE_BRUSH("Icons/GeneralTools/Paint_40x", Icon20x20));
+		Set("MeshPaint.PaintVertexWeight", new IMAGE_BRUSH("Icons/GeneralTools/Paint_40x", Icon20x20));
+		Set("MeshPaint.PaintVertexWeight.Small", new IMAGE_BRUSH("Icons/GeneralTools/Paint_40x", Icon20x20));
+		Set("MeshPaint.PaintTextureColor", new IMAGE_BRUSH("Icons/GeneralTools/Paint_40x", Icon20x20));
+		Set("MeshPaint.PaintTextureColor.Small", new IMAGE_BRUSH("Icons/GeneralTools/Paint_40x", Icon20x20));
+		Set("MeshPaint.PaintTextureAsset", new IMAGE_BRUSH("Icons/GeneralTools/Paint_40x", Icon20x20));
+		Set("MeshPaint.PaintTextureAsset.Small", new IMAGE_BRUSH("Icons/GeneralTools/Paint_40x", Icon20x20));
+		Set("MeshPaint.SwapColor", new IMAGE_BRUSH("Icons/Paint/Paint_SwapColors_40x", Icon20x20));
+		Set("MeshPaint.SwapColor.Small", new IMAGE_BRUSH("Icons/Paint/Paint_SwapColors_40x", Icon20x20));
+		Set("MeshPaint.FillVertex", new IMAGE_BRUSH("/Icons/GeneralTools/PaintBucket_40x", Icon20x20));
+		Set("MeshPaint.FillVertex.Small", new IMAGE_BRUSH("/Icons/GeneralTools/PaintBucket_40x", Icon20x20));
+		Set("MeshPaint.FillTexture", new IMAGE_BRUSH("/Icons/GeneralTools/PaintBucket_40x", Icon20x20));
+		Set("MeshPaint.FillTexture.Small", new IMAGE_BRUSH("/Icons/GeneralTools/PaintBucket_40x", Icon20x20));
+		Set("MeshPaint.Apply", new IMAGE_BRUSH("/Icons/Paint/Paint_Propagate_40x", Icon20x20));
+		Set("MeshPaint.Apply.Small", new IMAGE_BRUSH("/Icons/Paint/Paint_Propagate_40x", Icon20x20));
+		Set("MeshPaint.Commit", new IMAGE_BRUSH("/Icons/Paint/Paint_PropagateTexture_40x", Icon20x20));
+		Set("MeshPaint.Commit.Small", new IMAGE_BRUSH("/Icons/Paint/Paint_PropagateTexture_40x", Icon20x20));
+		Set("MeshPaint.Propagate", new IMAGE_BRUSH("Icons/Paint/Paint_AllLODs_40x", Icon20x20));
+		Set("MeshPaint.Propagate.Small", new IMAGE_BRUSH("Icons/Paint/Paint_AllLODs_40x", Icon20x20));
+		Set("MeshPaint.SaveVertex", new IMAGE_BRUSH("/Icons/GeneralTools/Save_40x", Icon20x20));
+		Set("MeshPaint.SaveVertex.Small", new IMAGE_BRUSH("/Icons/GeneralTools/Save_40x", Icon20x20));
+		Set("MeshPaint.SaveTexture", new IMAGE_BRUSH("/Icons/GeneralTools/Save_40x", Icon20x20));
+		Set("MeshPaint.SaveTexture.Small", new IMAGE_BRUSH("/Icons/GeneralTools/Save_40x", Icon20x20));
+		Set("MeshPaint.Add", new IMAGE_BRUSH("/Icons/icon_add_40x", Icon20x20));
+		Set("MeshPaint.Add.Small", new IMAGE_BRUSH("/Icons/icon_add_40x", Icon20x20));
+		Set("MeshPaint.RemoveVertex", new IMAGE_BRUSH("/Icons/GeneralTools/Delete_40x", Icon20x20));
+		Set("MeshPaint.RemoveVertex.Small", new IMAGE_BRUSH("/Icons/GeneralTools/Delete_40x", Icon20x20));
+		Set("MeshPaint.RemoveTexture", new IMAGE_BRUSH("/Icons/GeneralTools/Delete_40x", Icon20x20));
+		Set("MeshPaint.RemoveTexture.Small", new IMAGE_BRUSH("/Icons/GeneralTools/Delete_40x", Icon20x20));
+		Set("MeshPaint.Copy", new IMAGE_BRUSH("/Icons/GeneralTools/Copy_40x", Icon20x20));
+		Set("MeshPaint.Copy.Small", new IMAGE_BRUSH("/Icons/GeneralTools/Copy_40x", Icon20x20));
+		Set("MeshPaint.Paste", new IMAGE_BRUSH("/Icons/GeneralTools/Paste_40x", Icon20x20));
+		Set("MeshPaint.Paste.Small", new IMAGE_BRUSH("/Icons/GeneralTools/Paste_40x", Icon20x20));
+		Set("MeshPaint.Import", new IMAGE_BRUSH("/Icons/GeneralTools/Import_40x", Icon20x20));
+		Set("MeshPaint.Import.Small", new IMAGE_BRUSH("/Icons/GeneralTools/Import_40x", Icon20x20));
+		Set("MeshPaint.GetTextureColors", new IMAGE_BRUSH("/Icons/GeneralTools/Import_40x", Icon20x20));
+		Set("MeshPaint.GetTextureColors.Small", new IMAGE_BRUSH("/Icons/GeneralTools/Import_40x", Icon20x20));
+		Set("MeshPaint.GetVertexColors", new IMAGE_BRUSH("/Icons/GeneralTools/Import_40x", Icon20x20));
+		Set("MeshPaint.GetVertexColors.Small", new IMAGE_BRUSH("/Icons/GeneralTools/Import_40x", Icon20x20));
+		Set("MeshPaint.Fix", new IMAGE_BRUSH("/Icons/GeneralTools/Fix_40x", Icon20x20));
+		Set("MeshPaint.Fix.Small", new IMAGE_BRUSH("/Icons/GeneralTools/Fix_40x", Icon20x20));
+		Set("MeshPaint.PreviousLOD", new IMAGE_BRUSH(TEXT("Icons/GeneralTools/Previous_40x"), Icon20x20));
+		Set("MeshPaint.PreviousLOD.Small", new IMAGE_BRUSH(TEXT("Icons/GeneralTools/Previous_40x"), Icon20x20));
+		Set("MeshPaint.NextLOD", new IMAGE_BRUSH(TEXT("Icons/GeneralTools/Next_40x"), Icon20x20));
+		Set("MeshPaint.NextLOD.Small", new IMAGE_BRUSH(TEXT("Icons/GeneralTools/Next_40x"), Icon20x20));
+		Set("MeshPaint.PreviousTexture", new IMAGE_BRUSH(TEXT("Icons/GeneralTools/Previous_40x"), Icon20x20));
+		Set("MeshPaint.PreviousTexture.Small", new IMAGE_BRUSH(TEXT("Icons/GeneralTools/Previous_40x"), Icon20x20));
+		Set("MeshPaint.NextTexture", new IMAGE_BRUSH(TEXT("Icons/GeneralTools/Next_40x"), Icon20x20));
+		Set("MeshPaint.NextTexture.Small", new IMAGE_BRUSH(TEXT("Icons/GeneralTools/Next_40x"), Icon20x20));
+		Set("MeshPaint.Brush", new IMAGE_BRUSH_SVG("Starship/Common/Paintbrush", Icon20x20));
+		Set("MeshPaint.FindInCB", new IMAGE_BRUSH("/Icons/icon_toolbar_genericfinder_40px", Icon20x20));
+		Set("MeshPaint.Swap", new IMAGE_BRUSH("/Icons/icon_MeshPaint_Swap_16x", Icon12x12));
 	}
 
 	// Scalability (Performance Warning)
