Attempt to fix some Texture code that's not using correct PreEdit/PostEdit pattern

diff --git a/Engine/Plugins/Enterprise/MDLImporter/Source/MDLImporter/Private/generator/MaterialTextureFactory.cpp b/Engine/Plugins/Enterprise/MDLImporter/Source/MDLImporter/Private/generator/MaterialTextureFactory.cpp
--- a/Engine/Plugins/Enterprise/MDLImporter/Source/MDLImporter/Private/generator/MaterialTextureFactory.cpp
+++ b/Engine/Plugins/Enterprise/MDLImporter/Source/MDLImporter/Private/generator/MaterialTextureFactory.cpp
@@ -156,6 +156,7 @@ namespace Generator
 
 		if (bSettingsChanged)
 		{
+			Texture->PreEditChange(nullptr);
 			Texture->MipGenSettings = MipGenSettings;
 			Texture->CompressionNoAlpha = Property.bCompressionNoAlpha;
 			Texture->CompressionSettings = Property.CompressionSettings;
@@ -165,7 +166,6 @@ namespace Generator
 			Texture->LODGroup = Property.LODGroup;
 			Texture->SRGB = bIsSRGB;
 			Texture->bFlipGreenChannel = Property.bFlipGreenChannel;
-			Texture->UpdateResource();
 			Texture->PostEditChange();
 		}
 	}
diff --git a/Engine/Plugins/Experimental/AnimToTexture/Source/AnimToTextureEditor/Public/AnimToTextureUtils.h b/Engine/Plugins/Experimental/AnimToTexture/Source/AnimToTextureEditor/Public/AnimToTextureUtils.h
--- a/Engine/Plugins/Experimental/AnimToTexture/Source/AnimToTextureEditor/Public/AnimToTextureUtils.h
+++ b/Engine/Plugins/Experimental/AnimToTexture/Source/AnimToTextureEditor/Public/AnimToTextureUtils.h
@@ -209,6 +209,13 @@ FORCEINLINE_DEBUGGABLE bool AnimToTexture_Private::WriteToTexture(
 {
 	check(Texture);
 
+	// TODO FIX ME: this is the wrong way to build a UTexture
+	// instead fill out an FImage
+	// and use Texture.Source.Init() from FImage
+	// there is no reason to be touching the PlatformData or BulkData
+	
+	Texture->PreEditChange(nullptr);
+
 	// ------------------------------------------------------------------------
 	// Get Texture Platform
 	FTexturePlatformData* PlatformData = Texture->GetPlatformData();
@@ -264,8 +271,7 @@ FORCEINLINE_DEBUGGABLE bool AnimToTexture_Private::WriteToTexture(
 	Texture->MipGenSettings = TextureMipGenSettings::TMGS_NoMipmaps;
 
 	// Update and Mark to Save.
-	Texture->UpdateResource();
-	Texture->MarkPackageDirty();
+	Texture->PostEditChange();
 
 	return true;
 }
diff --git a/Engine/Plugins/MeshPainting/Source/MeshPaintingToolset/Private/MeshTexturePaintingTool.cpp b/Engine/Plugins/MeshPainting/Source/MeshPaintingToolset/Private/MeshTexturePaintingTool.cpp
--- a/Engine/Plugins/MeshPainting/Source/MeshPaintingToolset/Private/MeshTexturePaintingTool.cpp
+++ b/Engine/Plugins/MeshPainting/Source/MeshPaintingToolset/Private/MeshTexturePaintingTool.cpp
@@ -1422,9 +1422,7 @@ void UMeshTexturePaintingTool::FinishPaintingTexture()
 			RenderTargetResource->ReadPixels(TexturePixels);
 
 			{
-				// For undo
-				// @@ use PreEdit/PostEdit instead
-				TextureData->ScratchTexture->Modify();
+				TextureData->ScratchTexture->PreEditChange(nullptr);
 
 				const int32 NumPixels = TexturePixels.Num();
 
@@ -1446,6 +1444,7 @@ void UMeshTexturePaintingTool::FinishPaintingTexture()
 				// If render target gamma used was 1.0 then disable SRGB for the static texture
 				TextureData->ScratchTexture->SRGB = FMath::Abs(RenderTargetResource->GetDisplayGamma() - 1.0f) >= KINDA_SMALL_NUMBER;
 				TextureData->ScratchTexture->bHasBeenPaintedInEditor = true;
+				TextureData->ScratchTexture->PostEditChange();
 			}
 		}
 
@@ -1643,9 +1642,8 @@ void UMeshTexturePaintingTool::CommitAllPaintedTextures()
 
 				{
 					// For undo
-					// @@ use PreEdit/PostEdit instead
 					TextureData->PaintingTexture2D->SetFlags(RF_Transactional);
-					TextureData->PaintingTexture2D->Modify();
+					TextureData->PaintingTexture2D->PreEditChange(nullptr);
 
 					// Store source art
 					// @@ use FMipLock.Image instead
diff --git a/Engine/Plugins/PCG/Source/PCG/Private/Data/PCGTextureData.cpp b/Engine/Plugins/PCG/Source/PCG/Private/Data/PCGTextureData.cpp
--- a/Engine/Plugins/PCG/Source/PCG/Private/Data/PCGTextureData.cpp
+++ b/Engine/Plugins/PCG/Source/PCG/Private/Data/PCGTextureData.cpp
@@ -457,10 +457,11 @@ void UPCGTextureData::InitializeInternal(UTexture* InTexture, uint32 InTextureIn
 			// Duplicate texture and change access flags (editor only). This duplicate texture will be used by the normal logic below.
 			FObjectDuplicationParameters DuplicationParams(Texture2D, /*Outer=*/this);
 			DuplicateTexture = CastChecked<UTexture2D>(StaticDuplicateObjectEx(DuplicationParams));
+			DuplicateTexture->PreEditChange(nullptr);
 			DuplicateTexture->MipGenSettings = TextureMipGenSettings::TMGS_NoMipmaps;
 			DuplicateTexture->CompressionSettings = TextureCompressionSettings::TC_VectorDisplacementmap; // Allows the texture to be in a non-compressed format (B8G8R8A8), which is necessary to convince the data to remain CPU-side.
 			DuplicateTexture->SRGB = false;
-			DuplicateTexture->UpdateResource();
+			DuplicateTexture->PostEditChange();
 		}
 	}
 #endif
