Fix localized text name duplication.
#rnx

diff --git a/Engine/Plugins/MeshPainting/Source/MeshPaintEditorMode/Private/MeshPaintMode.cpp b/Engine/Plugins/MeshPainting/Source/MeshPaintEditorMode/Private/MeshPaintMode.cpp
--- a/Engine/Plugins/MeshPainting/Source/MeshPaintEditorMode/Private/MeshPaintMode.cpp
+++ b/Engine/Plugins/MeshPainting/Source/MeshPaintEditorMode/Private/MeshPaintMode.cpp
@@ -1218,7 +1218,7 @@ bool UMeshPaintMode::CanImportVertexColorsFromFile() const
 
 void UMeshPaintMode::ImportVertexColorsFromMeshPaintTexture()
 {
-	FScopedTransaction Transaction(LOCTEXT("LevelMeshPainter_TransactionImportColors", "Importing Vertex Colors From Mesh Paint Textures"));
+	FScopedTransaction Transaction(LOCTEXT("LevelMeshPainter_TransactionImportVertexColorFromTextureColor", "Importing Vertex Colors From Mesh Paint Textures"));
 	const TArray<UMeshComponent*> MeshComponents = GetSelectedComponents<UMeshComponent>();
 	for (UMeshComponent* Component : MeshComponents)
 	{
@@ -1241,7 +1241,7 @@ bool UMeshPaintMode::CanImportVertexColorsFromMeshPaintTexture() const
 
 void UMeshPaintMode::ImportMeshPaintTextureFromVertexColors()
 {
-	FScopedTransaction Transaction(LOCTEXT("LevelMeshPainter_TransactionImportTextures", "Importing Mesh Paint Textures From Vertex Colors"));
+	FScopedTransaction Transaction(LOCTEXT("LevelMeshPainter_TransactionImportTextureColorFromVertexColor", "Importing Mesh Paint Textures From Vertex Colors"));
 	const TArray<UMeshComponent*> MeshComponents = GetSelectedComponents<UMeshComponent>();
 	for (UMeshComponent* Component : MeshComponents)
 	{
