Fix initialization error in FPaintTexture2DData
#rnx

diff --git a/Engine/Plugins/MeshPainting/Source/MeshPaintingToolset/Public/MeshPaintingToolsetTypes.h b/Engine/Plugins/MeshPainting/Source/MeshPaintingToolset/Public/MeshPaintingToolsetTypes.h
--- a/Engine/Plugins/MeshPainting/Source/MeshPaintingToolset/Public/MeshPaintingToolsetTypes.h
+++ b/Engine/Plugins/MeshPainting/Source/MeshPaintingToolset/Public/MeshPaintingToolsetTypes.h
@@ -135,18 +135,18 @@ struct FPaintTexture2DData
 
 	/** The original texture that we're painting */
 	UPROPERTY(Transient)
-	TObjectPtr<UTexture2D> PaintingTexture2D = nullptr;
+	TObjectPtr<UTexture2D> PaintingTexture2D;
 
 	UPROPERTY(Transient)
 	bool bIsPaintingTexture2DModified = false;
 
 	/** Texture used to store the paint modifications for the transactions */
 	UPROPERTY(Transient)
-	TObjectPtr<UTexture2D> ScratchTexture = nullptr;
+	TObjectPtr<UTexture2D> ScratchTexture;
 
 	/** Render target texture for painting */
 	UPROPERTY(Transient)
-	TObjectPtr<UTextureRenderTarget2D> PaintRenderTargetTexture = nullptr;
+	TObjectPtr<UTextureRenderTarget2D> PaintRenderTargetTexture;
 
 	/** List of component we are painting on */
 	UPROPERTY(Transient)
@@ -154,11 +154,11 @@ struct FPaintTexture2DData
 	
 	/** Optional render target texture used as an input while painting that contains a clone of the texture painting brush */
 	UPROPERTY(Transient)
-	TObjectPtr<UTextureRenderTarget2D> PaintBrushRenderTargetTexture = nullptr;
+	TObjectPtr<UTextureRenderTarget2D> PaintBrushRenderTargetTexture;
 
 	/** Optional virtual texture adapter that we can use to visualize the render target in materials that use mesh paint virtual textures */
 	UPROPERTY(Transient)
-	TObjectPtr<class UVirtualTextureAdapter> PaintRenderTargetTextureAdapter = nullptr;
+	TObjectPtr<class UVirtualTextureAdapter> PaintRenderTargetTextureAdapter;
 
 	/** Temporary render target used to draw incremental paint to */
 	UPROPERTY(Transient)
@@ -170,11 +170,11 @@ struct FPaintTexture2DData
 
 	/** Optional render target used to store generated mask for texture seams. We create this by projecting object triangles into texture space using the selected UV channel. */
 	UPROPERTY(Transient)
-	TObjectPtr<UTextureRenderTarget2D> SeamMaskRenderTargetTexture = nullptr;
+	TObjectPtr<UTextureRenderTarget2D> SeamMaskRenderTargetTexture;
 
 	/** True if we need to generate a texture seam mask used for texture dilation */
 	UPROPERTY(Transient)
-	bool bGenerateSeamMask;
+	bool bGenerateSeamMask = false;
 
 	FPaintTexture2DData() = default;
 
