MeshPaintTool: Respect CanMeshPaintTextureColors() to disable selection/adding/painting.
Split the texture selection tool in two for texture color and texture asset painting.
New texture color mode respects the component/mesh disabling texture color painting.

diff --git a/Engine/Plugins/MeshPainting/Source/MeshPaintEditorMode/Private/MeshPaintMode.cpp b/Engine/Plugins/MeshPainting/Source/MeshPaintEditorMode/Private/MeshPaintMode.cpp
--- a/Engine/Plugins/MeshPainting/Source/MeshPaintEditorMode/Private/MeshPaintMode.cpp
+++ b/Engine/Plugins/MeshPainting/Source/MeshPaintEditorMode/Private/MeshPaintMode.cpp
@@ -36,7 +36,8 @@ FName UMeshPaintMode::MeshPaintMode_TextureColor = FName(TEXT("TextureColor"));
 FName UMeshPaintMode::MeshPaintMode_TextureAsset = FName(TEXT("Texture"));
 
 FString UMeshPaintMode::VertexSelectToolName = TEXT("VertexAdapterClickTool");
-FString UMeshPaintMode::TextureSelectToolName = TEXT("TextureAdapterClickTool");
+FString UMeshPaintMode::TextureColorSelectToolName = TEXT("TextureColorAdapterClickTool");
+FString UMeshPaintMode::TextureAssetSelectToolName = TEXT("TextureAssetAdapterClickTool");
 
 FString UMeshPaintMode::VertexColorPaintToolName = TEXT("VertexColorBrushTool");
 FString UMeshPaintMode::VertexWeightPaintToolName = TEXT("VertexWeightBrushTool");
@@ -130,8 +131,11 @@ void UMeshPaintMode::Enter()
 	UVertexAdapterClickToolBuilder* VertexClickToolBuilder = NewObject<UVertexAdapterClickToolBuilder>(this);
 	RegisterTool(ToolManagerCommands.SelectVertex, VertexSelectToolName, VertexClickToolBuilder);
 
-	UTextureAdapterClickToolBuilder* TextureClickToolBuilder = NewObject<UTextureAdapterClickToolBuilder>(this);
-	RegisterTool(ToolManagerCommands.SelectTexture, TextureSelectToolName, TextureClickToolBuilder);
+	UTextureColorAdapterClickToolBuilder* TextureColorClickToolBuilder = NewObject<UTextureColorAdapterClickToolBuilder>(this);
+	RegisterTool(ToolManagerCommands.SelectTextureColor, TextureColorSelectToolName, TextureColorClickToolBuilder);
+
+	UTextureAssetAdapterClickToolBuilder* TextureAssetClickToolBuilder = NewObject<UTextureAssetAdapterClickToolBuilder>(this);
+	RegisterTool(ToolManagerCommands.SelectTextureAsset, TextureAssetSelectToolName, TextureAssetClickToolBuilder);
 
 	UMeshColorPaintingToolBuilder* MeshColorPaintingToolBuilder = NewObject<UMeshColorPaintingToolBuilder>(this);
 	RegisterTool(ToolManagerCommands.PaintVertexColor, VertexColorPaintToolName, MeshColorPaintingToolBuilder);
@@ -437,13 +441,27 @@ void UMeshPaintMode::UpdateToolForSelection(const TArray<UMeshComponent*>& Curre
 	}
 
 	GEngine->GetEngineSubsystem<UMeshPaintingSubsystem>()->SetSelectionHasMaterialValidForTexturePaint(bCurrentSelectionSupportsTexturePaint);
-	if (!bCurrentSelectionSupportsTexturePaint 
-		&& GetToolManager() 
-		&& GetToolManager()->GetActiveTool(EToolSide::Mouse)
-		&& GetToolManager()->GetActiveTool(EToolSide::Mouse)->IsA<UMeshTexturePaintingTool>())
+
+	if (!bCurrentSelectionSupportsTexturePaint)
 	{
-		GetInteractiveToolsContext()->EndTool(EToolShutdownType::Accept);
-		GetInteractiveToolsContext()->StartTool(TextureSelectToolName);
+		UInteractiveToolManager* ToolManager = GetToolManager();
+		UInteractiveTool const* Tool = (ToolManager == nullptr) ? nullptr : ToolManager->GetActiveTool(EToolSide::Mouse);
+		if (Tool != nullptr)
+		{
+			if (Tool->IsA<UMeshTexturePaintingTool>())
+			{
+				GetInteractiveToolsContext()->EndTool(EToolShutdownType::Accept);
+
+				if (Tool->IsA<UMeshTextureColorPaintingTool>())
+				{
+					GetInteractiveToolsContext()->StartTool(TextureColorSelectToolName);
+				}
+				else if (Tool->IsA<UMeshTextureAssetPaintingTool>())
+				{
+					GetInteractiveToolsContext()->StartTool(TextureAssetSelectToolName);
+				}
+			}
+		}
 	}
 }
 
@@ -914,9 +932,13 @@ void UMeshPaintMode::ActivateDefaultTool()
 	{
 		GetInteractiveToolsContext()->StartTool(VertexSelectToolName);
 	}
-	if (PaletteName == UMeshPaintMode::MeshPaintMode_TextureColor || PaletteName == UMeshPaintMode::MeshPaintMode_TextureAsset)
+	if (PaletteName == UMeshPaintMode::MeshPaintMode_TextureColor)
+	{
+		GetInteractiveToolsContext()->StartTool(TextureColorSelectToolName);
+	}
+	if (PaletteName == UMeshPaintMode::MeshPaintMode_TextureAsset)
 	{
-		GetInteractiveToolsContext()->StartTool(TextureSelectToolName);
+		GetInteractiveToolsContext()->StartTool(TextureAssetSelectToolName);
 	}
 }
 
@@ -939,12 +961,12 @@ void UMeshPaintMode::UpdateOnPaletteChange(FName NewPaletteName)
 	else if (NewPaletteName == UMeshPaintMode::MeshPaintMode_TextureColor)
 	{
 		SwitchToToolPaint = TextureColorPaintToolName;
-		SwitchToToolSelect = TextureSelectToolName;
+		SwitchToToolSelect = TextureColorSelectToolName;
 	}
 	else if (NewPaletteName == UMeshPaintMode::MeshPaintMode_TextureAsset)
 	{
 		SwitchToToolPaint = TextureAssetPaintToolName;
-		SwitchToToolSelect = TextureSelectToolName;
+		SwitchToToolSelect = TextureAssetSelectToolName;
 	}
 
 	if (!SwitchToToolPaint.IsEmpty())
@@ -988,7 +1010,7 @@ void UMeshPaintMode::OnResetViewMode()
 bool UMeshPaintMode::IsInSelectTool() const
 {
 	FString ActiveTool = GetToolManager()->GetActiveToolName(EToolSide::Mouse);
-	return ActiveTool == VertexSelectToolName || ActiveTool == TextureSelectToolName;
+	return ActiveTool == VertexSelectToolName || ActiveTool == TextureColorSelectToolName || ActiveTool == TextureAssetSelectToolName;
 }
 
 bool UMeshPaintMode::IsInPaintTool() const
@@ -1061,7 +1083,7 @@ bool UMeshPaintMode::CanAddMeshPaintTextures() const
 		const TArray<UStaticMeshComponent*> StaticMeshComponents = GetSelectedComponents<UStaticMeshComponent>();
 		for (UStaticMeshComponent* Component : StaticMeshComponents)
 		{
-			if (Component->GetMeshPaintTexture() == nullptr)
+			if (Component->GetMeshPaintTexture() == nullptr && Component->CanMeshPaintTextureColors())
 			{
 				return true;
 			}
diff --git a/Engine/Plugins/MeshPainting/Source/MeshPaintEditorMode/Private/MeshPaintMode.h b/Engine/Plugins/MeshPainting/Source/MeshPaintEditorMode/Private/MeshPaintMode.h
--- a/Engine/Plugins/MeshPainting/Source/MeshPaintEditorMode/Private/MeshPaintMode.h
+++ b/Engine/Plugins/MeshPainting/Source/MeshPaintEditorMode/Private/MeshPaintMode.h
@@ -47,7 +47,8 @@ public:
 	static FName MeshPaintMode_TextureColor;
 	static FName MeshPaintMode_TextureAsset;
 	static FString VertexSelectToolName;
-	static FString TextureSelectToolName;
+	static FString TextureColorSelectToolName;
+	static FString TextureAssetSelectToolName;
 	static FString VertexColorPaintToolName;
 	static FString VertexWeightPaintToolName;
 	static FString TextureColorPaintToolName;
diff --git a/Engine/Plugins/MeshPainting/Source/MeshPaintEditorMode/Private/MeshPaintModeCommands.cpp b/Engine/Plugins/MeshPainting/Source/MeshPaintEditorMode/Private/MeshPaintModeCommands.cpp
--- a/Engine/Plugins/MeshPainting/Source/MeshPaintEditorMode/Private/MeshPaintModeCommands.cpp
+++ b/Engine/Plugins/MeshPainting/Source/MeshPaintEditorMode/Private/MeshPaintModeCommands.cpp
@@ -32,7 +32,8 @@ void FMeshPaintingToolActionCommands::UpdateToolCommandBinding(UInteractiveTool*
 void FMeshPaintEditorModeCommands::RegisterCommands()
 {
 	UI_COMMAND(SelectVertex, "Select", "Select the mesh for vertex painting", EUserInterfaceActionType::ToggleButton, FInputChord());
-	UI_COMMAND(SelectTexture, "Select", "Select the mesh for texture painting", EUserInterfaceActionType::ToggleButton, FInputChord());
+	UI_COMMAND(SelectTextureColor, "Select", "Select the mesh for texture color painting", EUserInterfaceActionType::ToggleButton, FInputChord());
+	UI_COMMAND(SelectTextureAsset, "Select", "Select the mesh for texture asset painting", EUserInterfaceActionType::ToggleButton, FInputChord());
 	UI_COMMAND(PaintVertexColor, "Paint", "Paint the mesh vertex colors", EUserInterfaceActionType::ToggleButton, FInputChord());
 	UI_COMMAND(PaintVertexWeight, "Paint", "Paint the mesh vertex weights", EUserInterfaceActionType::ToggleButton, FInputChord());
 	UI_COMMAND(PaintTextureColor, "Paint", "Paint the mesh texture colors", EUserInterfaceActionType::ToggleButton, FInputChord());
@@ -70,12 +71,12 @@ void FMeshPaintEditorModeCommands::RegisterCommands()
 	Commands.Add(UMeshPaintMode::MeshPaintMode_VertexWeights, VertexWeightCommands);
 
 	TArray<TSharedPtr<FUICommandInfo>> TextureColorCommands = {
-		SelectTexture, PaintTextureColor, SwapColor, FillTexture, Commit, SaveTexture, Add, RemoveTexture, Copy, Paste, GetVertexColors
+		SelectTextureColor, PaintTextureColor, SwapColor, FillTexture, Commit, SaveTexture, Add, RemoveTexture, Copy, Paste, GetVertexColors
 	};
 	Commands.Add(UMeshPaintMode::MeshPaintMode_TextureColor, TextureColorCommands);
 
 	TArray<TSharedPtr<FUICommandInfo>> TextureAssetCommands = {
-		SelectTexture, PaintTextureAsset, SwapColor, FillTexture, Commit, SaveTexture
+		SelectTextureAsset, PaintTextureAsset, SwapColor, FillTexture, Commit, SaveTexture
 	};
 	Commands.Add(UMeshPaintMode::MeshPaintMode_TextureAsset, TextureAssetCommands);
 }
diff --git a/Engine/Plugins/MeshPainting/Source/MeshPaintEditorMode/Public/MeshPaintModeCommands.h b/Engine/Plugins/MeshPainting/Source/MeshPaintEditorMode/Public/MeshPaintModeCommands.h
--- a/Engine/Plugins/MeshPainting/Source/MeshPaintEditorMode/Public/MeshPaintModeCommands.h
+++ b/Engine/Plugins/MeshPainting/Source/MeshPaintEditorMode/Public/MeshPaintModeCommands.h
@@ -31,7 +31,8 @@ public:
 
 public:
 	TSharedPtr<FUICommandInfo> SelectVertex;
-	TSharedPtr<FUICommandInfo> SelectTexture;
+	TSharedPtr<FUICommandInfo> SelectTextureColor;
+	TSharedPtr<FUICommandInfo> SelectTextureAsset;
 
 	TSharedPtr<FUICommandInfo> PaintVertexColor;
 	TSharedPtr<FUICommandInfo> PaintVertexWeight;
diff --git a/Engine/Plugins/MeshPainting/Source/MeshPaintingToolset/Private/MeshPaintHelpers.cpp b/Engine/Plugins/MeshPainting/Source/MeshPaintingToolset/Private/MeshPaintHelpers.cpp
--- a/Engine/Plugins/MeshPainting/Source/MeshPaintingToolset/Private/MeshPaintHelpers.cpp
+++ b/Engine/Plugins/MeshPainting/Source/MeshPaintingToolset/Private/MeshPaintHelpers.cpp
@@ -77,7 +77,7 @@ void UMeshPaintingSubsystem::RemoveComponentInstanceVertexColors(UStaticMeshComp
 
 void UMeshPaintingSubsystem::CreateComponentMeshPaintTexture(UStaticMeshComponent* StaticMeshComponent)
 {
-	if (StaticMeshComponent != nullptr && StaticMeshComponent->GetMeshPaintTexture() == nullptr)
+	if (StaticMeshComponent != nullptr && StaticMeshComponent->GetMeshPaintTexture() == nullptr && StaticMeshComponent->CanMeshPaintTextureColors())
 	{
 		StaticMeshComponent->Modify();
 
diff --git a/Engine/Plugins/MeshPainting/Source/MeshPaintingToolset/Private/MeshSelect.cpp b/Engine/Plugins/MeshPainting/Source/MeshPaintingToolset/Private/MeshSelect.cpp
--- a/Engine/Plugins/MeshPainting/Source/MeshPaintingToolset/Private/MeshSelect.cpp
+++ b/Engine/Plugins/MeshPainting/Source/MeshPaintingToolset/Private/MeshSelect.cpp
@@ -23,14 +23,25 @@ UInteractiveTool* UVertexAdapterClickToolBuilder::BuildTool(const FToolBuilderSt
 	return NewTool;
 }
 
-bool UTextureAdapterClickToolBuilder::CanBuildTool(const FToolBuilderState& SceneState) const
+bool UTextureColorAdapterClickToolBuilder::CanBuildTool(const FToolBuilderState& SceneState) const
 {
 	return true;
 }
 
-UInteractiveTool* UTextureAdapterClickToolBuilder::BuildTool(const FToolBuilderState& SceneState) const
+UInteractiveTool* UTextureColorAdapterClickToolBuilder::BuildTool(const FToolBuilderState& SceneState) const
 {
-	UTextureAdapterClickTool* NewTool = NewObject<UTextureAdapterClickTool>(SceneState.ToolManager);
+	UTextureColorAdapterClickTool* NewTool = NewObject<UTextureColorAdapterClickTool>(SceneState.ToolManager);
+	return NewTool;
+}
+
+bool UTextureAssetAdapterClickToolBuilder::CanBuildTool(const FToolBuilderState& SceneState) const
+{
+	return true;
+}
+
+UInteractiveTool* UTextureAssetAdapterClickToolBuilder::BuildTool(const FToolBuilderState& SceneState) const
+{
+	UTextureAssetAdapterClickTool* NewTool = NewObject<UTextureAssetAdapterClickTool>(SceneState.ToolManager);
 	return NewTool;
 }
 
@@ -77,24 +88,17 @@ void UMeshClickTool::OnClicked(const FInputDeviceRay& ClickPos)
 }
 
 
-UVertexAdapterClickTool::UVertexAdapterClickTool()
-	: UMeshClickTool()
-{
-
-}
-
 bool UVertexAdapterClickTool::IsMeshAdapterSupported(TSharedPtr<IMeshPaintComponentAdapter> MeshAdapter) const
 {
 	return MeshAdapter.IsValid() ? MeshAdapter->SupportsVertexPaint() : false;
 }
 
-UTextureAdapterClickTool::UTextureAdapterClickTool()
-	: UMeshClickTool()
+bool UTextureColorAdapterClickTool::IsMeshAdapterSupported(TSharedPtr<IMeshPaintComponentAdapter> MeshAdapter) const
 {
-
+	return MeshAdapter.IsValid() ? MeshAdapter->SupportsTextureColorPaint() : false;
 }
 
-bool UTextureAdapterClickTool::IsMeshAdapterSupported(TSharedPtr<IMeshPaintComponentAdapter> MeshAdapter) const
+bool UTextureAssetAdapterClickTool::IsMeshAdapterSupported(TSharedPtr<IMeshPaintComponentAdapter> MeshAdapter) const
 {
 	return MeshAdapter.IsValid() ? MeshAdapter->SupportsTexturePaint() : false;
 }
diff --git a/Engine/Plugins/MeshPainting/Source/MeshPaintingToolset/Private/MeshTexturePaintingTool.cpp b/Engine/Plugins/MeshPainting/Source/MeshPaintingToolset/Private/MeshTexturePaintingTool.cpp
--- a/Engine/Plugins/MeshPainting/Source/MeshPaintingToolset/Private/MeshTexturePaintingTool.cpp
+++ b/Engine/Plugins/MeshPainting/Source/MeshPaintingToolset/Private/MeshTexturePaintingTool.cpp
@@ -244,8 +244,7 @@ void UMeshTexturePaintingTool::Render(IToolsContextRenderAPI* RenderAPI)
 		for (UMeshComponent* CurrentComponent : MeshPaintingSubsystem->GetPaintableMeshComponents())
 		{
 			TSharedPtr<IMeshPaintComponentAdapter> MeshAdapter = MeshPaintingSubsystem->GetAdapterForComponent(Cast<UMeshComponent>(CurrentComponent));
-
-			if (MeshAdapter->IsValid() && MeshAdapter->SupportsVertexPaint())
+			if (IsMeshAdapterSupported(MeshAdapter))
 			{
 				const FMatrix ComponentToWorldMatrix = MeshAdapter->GetComponentToWorldMatrix();
 				FViewCameraState CameraState;
@@ -299,57 +298,54 @@ void UMeshTexturePaintingTool::OnTick(float DeltaTime)
 				
 				const int32 UVChannel = GetSelectedUVChannel(SelectedComponent);
 
-				IMeshPaintComponentAdapter* MeshAdapter = MeshPaintingSubsystem->GetAdapterForComponent(SelectedComponent).Get();
-				if (MeshAdapter)
+				TSharedPtr<IMeshPaintComponentAdapter> MeshAdapter = MeshPaintingSubsystem->GetAdapterForComponent(SelectedComponent);
+				if (IsMeshAdapterSupported(MeshAdapter))
 				{
-					if (MeshAdapter->SupportsTexturePaint())
+					Textures.Empty();
+					const UTexture2D* TargetTexture2D = GetSelectedPaintTexture(SelectedComponent);
+					if (TargetTexture2D)
 					{
-						Textures.Empty();
-						const UTexture2D* TargetTexture2D = GetSelectedPaintTexture(SelectedComponent);
-						if (TargetTexture2D)
-						{
-							Textures.Add(TargetTexture2D);
-
-							FPaintTexture2DData* TextureData = GetPaintTargetData(TargetTexture2D);
-							if (TextureData)
-							{
-								Textures.Add(TextureData->PaintRenderTargetTexture);
-							}
+						Textures.Add(TargetTexture2D);
 
-							TArray<FTexturePaintMeshSectionInfo> MaterialSections;
-							UTexturePaintToolset::RetrieveMeshSectionsForTextures(SelectedComponent, 0 /*CachedLODIndex*/, Textures, MaterialSections);
+						FPaintTexture2DData* TextureData = GetPaintTargetData(TargetTexture2D);
+						if (TextureData)
+						{
+							Textures.Add(TextureData->PaintRenderTargetTexture);
+						}
 
-							TArray<FTexturePaintTriangleInfo> TrianglePaintInfoArray;
-							FPerTrianglePaintAction TempAction = FPerTrianglePaintAction::CreateUObject(this, &UMeshTexturePaintingTool::GatherTextureTriangles, &TrianglePaintInfoArray, &MaterialSections, UVChannel);
+						TArray<FTexturePaintMeshSectionInfo> MaterialSections;
+						UTexturePaintToolset::RetrieveMeshSectionsForTextures(SelectedComponent, 0 /*CachedLODIndex*/, Textures, MaterialSections);
 
-							// We are flooding the texture, so all triangles are influenced
+						TArray<FTexturePaintTriangleInfo> TrianglePaintInfoArray;
+						FPerTrianglePaintAction TempAction = FPerTrianglePaintAction::CreateUObject(this, &UMeshTexturePaintingTool::GatherTextureTriangles, &TrianglePaintInfoArray, &MaterialSections, UVChannel);
 
-							const TArray<uint32>& MeshIndices = MeshAdapter->GetMeshIndices();
-							int32 TriangleIndices[3];
+						// We are flooding the texture, so all triangles are influenced
 
-							for (int32 i = 0; i < MeshIndices.Num(); i += 3)
-							{
-								TriangleIndices[0] = MeshIndices[i + 0];
-								TriangleIndices[1] = MeshIndices[i + 1];
-								TriangleIndices[2] = MeshIndices[i + 2];
-								TempAction.Execute(MeshAdapter, i / 3, TriangleIndices);
-							}
+						const TArray<uint32>& MeshIndices = MeshAdapter->GetMeshIndices();
+						int32 TriangleIndices[3];
 
-							// Painting textures
-							if ((TexturePaintingCurrentMeshComponent != nullptr) && (TexturePaintingCurrentMeshComponent != SelectedComponent))
-							{
-								// Mesh has changed, so finish up with our previous texture
-								FinishPaintingTexture();
-							}
+						for (int32 i = 0; i < MeshIndices.Num(); i += 3)
+						{
+							TriangleIndices[0] = MeshIndices[i + 0];
+							TriangleIndices[1] = MeshIndices[i + 1];
+							TriangleIndices[2] = MeshIndices[i + 2];
+							TempAction.Execute(MeshAdapter.Get(), i / 3, TriangleIndices);
+						}
 
-							if (TexturePaintingCurrentMeshComponent == nullptr)
-							{
-								StartPaintingTexture(SelectedComponent, *MeshAdapter);
-							}
+						// Painting textures
+						if ((TexturePaintingCurrentMeshComponent != nullptr) && (TexturePaintingCurrentMeshComponent != SelectedComponent))
+						{
+							// Mesh has changed, so finish up with our previous texture
+							FinishPaintingTexture();
+						}
 
-							FMeshPaintParameters* LastParams = nullptr;
-							PaintTexture(BucketFillParams, UVChannel, TrianglePaintInfoArray, *MeshAdapter, LastParams);
+						if (TexturePaintingCurrentMeshComponent == nullptr)
+						{
+							StartPaintingTexture(SelectedComponent, *MeshAdapter);
 						}
+
+						FMeshPaintParameters* LastParams = nullptr;
+						PaintTexture(BucketFillParams, UVChannel, TrianglePaintInfoArray, *MeshAdapter, LastParams);
 					}
 				}
 			}
@@ -593,16 +589,12 @@ bool UMeshTexturePaintingTool::PaintInternal(const TArrayView<TPair<FVector, FVe
 		{
 			UMeshComponent* HoveredComponent = Entry.Key;
 			TArray<int32>& PaintRayResultIds = Entry.Value;
-			IMeshPaintComponentAdapter* MeshAdapter = MeshPaintingSubsystem ? MeshPaintingSubsystem->GetAdapterForComponent(HoveredComponent).Get() : nullptr;
-			if (!ensure(MeshAdapter))
+			
+			TSharedPtr<IMeshPaintComponentAdapter> MeshAdapter = MeshPaintingSubsystem->GetAdapterForComponent(HoveredComponent);
+			if (IsMeshAdapterSupported(MeshAdapter))
 			{
-				continue;
-			}
-
-			const int32 UVChannel = GetSelectedUVChannel(HoveredComponent);
+				const int32 UVChannel = GetSelectedUVChannel(HoveredComponent);
 
-			if (MeshAdapter->SupportsTexturePaint())
-			{
 				Textures.Empty();
 				const UTexture2D* TargetTexture2D = GetSelectedPaintTexture(HoveredComponent);
 				if (TargetTexture2D)
@@ -624,7 +616,7 @@ bool UMeshTexturePaintingTool::PaintInternal(const TArrayView<TPair<FVector, FVe
 						const FVector& BestTraceResultLocation = PaintRayResults[PaintRayResultId].BestTraceResult.Location;
 						FViewCameraState CameraState;
 						GetToolManager()->GetContextQueriesAPI()->GetCurrentViewState(CameraState);
-						bPaintApplied |= MeshPaintingSubsystem->ApplyPerTrianglePaintAction(MeshAdapter, CameraState.Position, BestTraceResultLocation, BrushProperties, FPerTrianglePaintAction::CreateUObject(this, &UMeshTexturePaintingTool::GatherTextureTriangles, &TrianglePaintInfoArray, &MaterialSections, UVChannel), TextureProperties->bOnlyFrontFacingTriangles);
+						bPaintApplied |= MeshPaintingSubsystem->ApplyPerTrianglePaintAction(MeshAdapter.Get(), CameraState.Position, BestTraceResultLocation, BrushProperties, FPerTrianglePaintAction::CreateUObject(this, &UMeshTexturePaintingTool::GatherTextureTriangles, &TrianglePaintInfoArray, &MaterialSections, UVChannel), TextureProperties->bOnlyFrontFacingTriangles);
 						break;
 					}
 
@@ -1481,9 +1473,9 @@ void UMeshTexturePaintingTool::OnTransactionStateChanged(const FTransactionConte
 
 								for (UMeshComponent* Component : TextureData->PaintedComponents)
 								{
-									if (TSharedPtr<IMeshPaintComponentAdapter> PaintAdapterRessource = MeshPaintingSubsystem->GetAdapterForComponent(Component))
+									if (TSharedPtr<IMeshPaintComponentAdapter> PaintAdapterResource = MeshPaintingSubsystem->GetAdapterForComponent(Component))
 									{
-										PaintAdapterRessource->ApplyOrRemoveTextureOverride(TextureData->PaintingTexture2D, GetTextureForMaterialOverride(*TextureData));
+										PaintAdapterResource->ApplyOrRemoveTextureOverride(TextureData->PaintingTexture2D, GetTextureForMaterialOverride(*TextureData));
 										PaintComponentOverride.PaintedComponents.Add(Component);
 									}
 								}
@@ -1497,9 +1489,9 @@ void UMeshTexturePaintingTool::OnTransactionStateChanged(const FTransactionConte
 
 							for (UMeshComponent* Component : ComponentOverride.PaintedComponents)
 							{
-								if (TSharedPtr<IMeshPaintComponentAdapter> PaintAdapterRessource = MeshPaintingSubsystem->GetAdapterForComponent(Component))
+								if (TSharedPtr<IMeshPaintComponentAdapter> PaintAdapterResource = MeshPaintingSubsystem->GetAdapterForComponent(Component))
 								{
-									PaintAdapterRessource->ApplyOrRemoveTextureOverride(Texture,  nullptr);
+									PaintAdapterResource->ApplyOrRemoveTextureOverride(Texture,  nullptr);
 								}
 							}
 						};
@@ -1521,9 +1513,9 @@ void UMeshTexturePaintingTool::OnTransactionStateChanged(const FTransactionConte
 
 								auto ComponentOverrideNotNeeded = [this, MeshPaintingSubsystem, &ComponentOverride, Texture, &MeshComponentsToRemove](TObjectPtr<UMeshComponent>& MeshComponent)
 								{
-									if (TSharedPtr<IMeshPaintComponentAdapter> PaintAdapterRessource = MeshPaintingSubsystem->GetAdapterForComponent(MeshComponent))
+									if (TSharedPtr<IMeshPaintComponentAdapter> PaintAdapterResource = MeshPaintingSubsystem->GetAdapterForComponent(MeshComponent))
 									{
-										PaintAdapterRessource->ApplyOrRemoveTextureOverride(Texture, nullptr);
+										PaintAdapterResource->ApplyOrRemoveTextureOverride(Texture, nullptr);
 									}
 
 									MeshComponentsToRemove.Add(MeshComponent);
@@ -1532,9 +1524,9 @@ void UMeshTexturePaintingTool::OnTransactionStateChanged(const FTransactionConte
 								UTexture* OverrideTexture = GetTextureForMaterialOverride(*TextureData);
 								auto ComponentOverrideMissing = [this, MeshPaintingSubsystem, &ComponentOverride, Texture, &MeshComponentsToAdd, OverrideTexture](TObjectPtr<UMeshComponent>& MeshComponent)
 								{
-									if (TSharedPtr<IMeshPaintComponentAdapter> PaintAdapterRessource = MeshPaintingSubsystem->GetAdapterForComponent(MeshComponent))
+									if (TSharedPtr<IMeshPaintComponentAdapter> PaintAdapterResource = MeshPaintingSubsystem->GetAdapterForComponent(MeshComponent))
 									{
-										PaintAdapterRessource->ApplyOrRemoveTextureOverride(Texture, OverrideTexture);
+										PaintAdapterResource->ApplyOrRemoveTextureOverride(Texture, OverrideTexture);
 									}
 
 									MeshComponentsToAdd.Add(MeshComponent);
@@ -1711,11 +1703,6 @@ int32 UMeshTexturePaintingTool::GetNumberOfPendingPaintChanges() const
 	return Result;
 }
 
-bool UMeshTexturePaintingTool::IsMeshAdapterSupported(TSharedPtr<IMeshPaintComponentAdapter> MeshAdapter) const
-{
-	return MeshAdapter.IsValid() ? MeshAdapter->SupportsTexturePaint() : false;
-}
-
 void UMeshTexturePaintingTool::FloodCurrentPaintTexture()
 {
 	bRequestPaintBucketFill = true;
@@ -1737,6 +1724,11 @@ void UMeshTextureColorPaintingTool::Setup()
 		EToolMessageLevel::UserNotification);
 }
 
+bool UMeshTextureColorPaintingTool::IsMeshAdapterSupported(TSharedPtr<IMeshPaintComponentAdapter> MeshAdapter) const
+{
+	return MeshAdapter.IsValid() ? MeshAdapter->SupportsTextureColorPaint() : false;
+}
+
 UTexture2D* UMeshTextureColorPaintingTool::GetSelectedPaintTexture(UMeshComponent const* InMeshComponent) const
 {
 	return Cast<UTexture2D>(InMeshComponent->GetMeshPaintTexture());
@@ -1793,6 +1785,11 @@ void UMeshTextureAssetPaintingTool::Setup()
 		EToolMessageLevel::UserNotification);
 }
 
+bool UMeshTextureAssetPaintingTool::IsMeshAdapterSupported(TSharedPtr<IMeshPaintComponentAdapter> MeshAdapter) const
+{
+	return MeshAdapter.IsValid() ? MeshAdapter->SupportsTexturePaint() : false;
+}
+
 UTexture2D* UMeshTextureAssetPaintingTool::GetSelectedPaintTexture(UMeshComponent const* InMeshComponent) const
 {
 	return AssetProperties->PaintTexture;
diff --git a/Engine/Plugins/MeshPainting/Source/MeshPaintingToolset/Public/IMeshPaintComponentAdapter.h b/Engine/Plugins/MeshPainting/Source/MeshPaintingToolset/Public/IMeshPaintComponentAdapter.h
--- a/Engine/Plugins/MeshPainting/Source/MeshPaintingToolset/Public/IMeshPaintComponentAdapter.h
+++ b/Engine/Plugins/MeshPainting/Source/MeshPaintingToolset/Public/IMeshPaintComponentAdapter.h
@@ -43,9 +43,12 @@ public:
 	/** Returns whether or not this adapter is in a valid state */
 	virtual bool IsValid() const = 0;	
 	
-	/** Whether or not this adapter supports texture painting */
+	/** Whether or not this adapter supports texture asset painting */
 	virtual bool SupportsTexturePaint() const = 0;
 
+	/** Whether or not this adapter supports texture color painting */
+	virtual bool SupportsTextureColorPaint() const = 0;
+
 	/** Whether or not this adapter supports vertex painting */
 	virtual bool SupportsVertexPaint() const = 0;
 
diff --git a/Engine/Plugins/MeshPainting/Source/MeshPaintingToolset/Public/MeshPaintGeometryCollectionAdapter.h b/Engine/Plugins/MeshPainting/Source/MeshPaintingToolset/Public/MeshPaintGeometryCollectionAdapter.h
--- a/Engine/Plugins/MeshPainting/Source/MeshPaintingToolset/Public/MeshPaintGeometryCollectionAdapter.h
+++ b/Engine/Plugins/MeshPainting/Source/MeshPaintingToolset/Public/MeshPaintGeometryCollectionAdapter.h
@@ -29,6 +29,7 @@ public:
 	virtual void OnRemoved() override;
 	virtual bool IsValid() const override;
 	virtual bool SupportsTexturePaint() const override { return IsValid(); }
+	virtual bool SupportsTextureColorPaint() const override { return false; }
 	virtual bool SupportsVertexPaint() const override { return IsValid(); }
 	virtual bool LineTraceComponent(struct FHitResult& OutHit, const FVector Start, const FVector End, const struct FCollisionQueryParams& Params) const override;
 	virtual void QueryPaintableTextures(int32 MaterialIndex, int32& OutDefaultIndex, TArray<struct FPaintableTexture>& InOutTextureList) override;
diff --git a/Engine/Plugins/MeshPainting/Source/MeshPaintingToolset/Public/MeshPaintSkeletalMeshAdapter.h b/Engine/Plugins/MeshPainting/Source/MeshPaintingToolset/Public/MeshPaintSkeletalMeshAdapter.h
--- a/Engine/Plugins/MeshPainting/Source/MeshPaintingToolset/Public/MeshPaintSkeletalMeshAdapter.h
+++ b/Engine/Plugins/MeshPainting/Source/MeshPaintingToolset/Public/MeshPaintSkeletalMeshAdapter.h
@@ -31,6 +31,7 @@ public:
 	virtual void OnRemoved() override;
 	virtual bool IsValid() const override { return SkeletalMeshComponent.IsValid() && ReferencedSkeletalMesh && SkeletalMeshComponent->GetSkeletalMeshAsset() == ReferencedSkeletalMesh; }
 	virtual bool SupportsTexturePaint() const override { return true; }
+	virtual bool SupportsTextureColorPaint() const override { return false; }
 	virtual bool SupportsVertexPaint() const override { return SkeletalMeshComponent.IsValid(); }
 	virtual bool LineTraceComponent(struct FHitResult& OutHit, const FVector Start, const FVector End, const struct FCollisionQueryParams& Params) const override;
 	virtual void QueryPaintableTextures(int32 MaterialIndex, int32& OutDefaultIndex, TArray<struct FPaintableTexture>& InOutTextureList) override;
diff --git a/Engine/Plugins/MeshPainting/Source/MeshPaintingToolset/Public/MeshPaintStaticMeshAdapter.h b/Engine/Plugins/MeshPainting/Source/MeshPaintingToolset/Public/MeshPaintStaticMeshAdapter.h
--- a/Engine/Plugins/MeshPainting/Source/MeshPaintingToolset/Public/MeshPaintStaticMeshAdapter.h
+++ b/Engine/Plugins/MeshPainting/Source/MeshPaintingToolset/Public/MeshPaintStaticMeshAdapter.h
@@ -28,6 +28,7 @@ public:
 	virtual void OnRemoved() override {}
 	virtual bool IsValid() const override { return StaticMeshComponent.IsValid() && ReferencedStaticMesh && StaticMeshComponent->GetStaticMesh() == ReferencedStaticMesh; }
 	virtual bool SupportsTexturePaint() const override { return true; }
+	virtual bool SupportsTextureColorPaint() const override { return StaticMeshComponent.IsValid() && StaticMeshComponent->CanMeshPaintTextureColors(); }
 	virtual bool SupportsVertexPaint() const override { return StaticMeshComponent.IsValid() && StaticMeshComponent->CanMeshPaintVertexColors(); }
 	virtual bool LineTraceComponent(struct FHitResult& OutHit, const FVector Start, const FVector End, const struct FCollisionQueryParams& Params) const override;	
 	virtual void QueryPaintableTextures(int32 MaterialIndex, int32& OutDefaultIndex, TArray<struct FPaintableTexture>& InOutTextureList) override;
diff --git a/Engine/Plugins/MeshPainting/Source/MeshPaintingToolset/Public/MeshSelect.h b/Engine/Plugins/MeshPainting/Source/MeshPaintingToolset/Public/MeshSelect.h
--- a/Engine/Plugins/MeshPainting/Source/MeshPaintingToolset/Public/MeshSelect.h
+++ b/Engine/Plugins/MeshPainting/Source/MeshPaintingToolset/Public/MeshSelect.h
@@ -26,10 +26,10 @@ public:
 
 
 /**
- * Builder for UTextureAdapterClickTool
+ * Builder for UTextureColorAdapterClickTool
  */
 UCLASS()
-class MESHPAINTINGTOOLSET_API UTextureAdapterClickToolBuilder : public USingleClickToolBuilder
+class MESHPAINTINGTOOLSET_API UTextureColorAdapterClickToolBuilder : public USingleClickToolBuilder
 {
 	GENERATED_BODY()
 
@@ -39,6 +39,19 @@ public:
 };
 
 
+/**
+ * Builder for UTextureAssetAdapterClickTool
+ */
+UCLASS()
+class MESHPAINTINGTOOLSET_API UTextureAssetAdapterClickToolBuilder : public USingleClickToolBuilder
+{
+	GENERATED_BODY()
+
+public:
+	virtual bool CanBuildTool(const FToolBuilderState& SceneState) const override;
+	virtual UInteractiveTool* BuildTool(const FToolBuilderState& SceneState) const override;
+};
+
 
 /**
  * USingleClickTool is perhaps the simplest possible interactive tool. It simply
@@ -81,7 +94,6 @@ protected:
 
 	UPROPERTY(Transient)
 	TObjectPtr<UMeshPaintSelectionMechanic> SelectionMechanic;
-
 };
 
 UCLASS()
@@ -90,28 +102,29 @@ class MESHPAINTINGTOOLSET_API UVertexAdapterClickTool : public UMeshClickTool
 	GENERATED_BODY()
 
 public:
-	UVertexAdapterClickTool();
-
 	// USingleClickTool overrides
 	virtual bool IsMeshAdapterSupported(TSharedPtr<IMeshPaintComponentAdapter> MeshAdapter) const override;
-
 };
 
 UCLASS()
-class MESHPAINTINGTOOLSET_API UTextureAdapterClickTool : public UMeshClickTool
+class MESHPAINTINGTOOLSET_API UTextureColorAdapterClickTool : public UMeshClickTool
 {
 	GENERATED_BODY()
 
 public:
-	UTextureAdapterClickTool();
-
 	// USingleClickTool overrides
 	virtual bool IsMeshAdapterSupported(TSharedPtr<IMeshPaintComponentAdapter> MeshAdapter) const override;
-	virtual bool AllowsMultiselect() const override
-	{
-		return true;
-	}
+};
 
+UCLASS()
+class MESHPAINTINGTOOLSET_API UTextureAssetAdapterClickTool : public UMeshClickTool
+{
+	GENERATED_BODY()
+
+public:
+	// USingleClickTool overrides
+	virtual bool IsMeshAdapterSupported(TSharedPtr<IMeshPaintComponentAdapter> MeshAdapter) const override;
+	virtual bool AllowsMultiselect() const override { return false; }
 };
 
 #if UE_ENABLE_INCLUDE_ORDER_DEPRECATED_IN_5_2
diff --git a/Engine/Plugins/MeshPainting/Source/MeshPaintingToolset/Public/MeshTexturePaintingTool.h b/Engine/Plugins/MeshPainting/Source/MeshPaintingToolset/Public/MeshTexturePaintingTool.h
--- a/Engine/Plugins/MeshPainting/Source/MeshPaintingToolset/Public/MeshTexturePaintingTool.h
+++ b/Engine/Plugins/MeshPainting/Source/MeshPaintingToolset/Public/MeshTexturePaintingTool.h
@@ -169,8 +169,6 @@ public:
 
 	void FloodCurrentPaintTexture();
 	
-	virtual bool IsMeshAdapterSupported(TSharedPtr<IMeshPaintComponentAdapter> MeshAdapter) const override;
-	
 	virtual void GetModifiedTexturesToSave(TArray<UObject*>& OutTexturesToSave) const {}
 
 protected:
@@ -281,6 +279,7 @@ public:
 
 	// Begin UMeshTexturePaintingTool Interface.
 	virtual bool AllowsMultiselect() const override { return true; }
+	virtual bool IsMeshAdapterSupported(TSharedPtr<IMeshPaintComponentAdapter> MeshAdapter) const override;
 	virtual UTexture2D* GetSelectedPaintTexture(UMeshComponent const* InMeshComponent) const override;
 	virtual int32 GetSelectedUVChannel(UMeshComponent const* InMeshComponent) const override;
 	virtual void GetModifiedTexturesToSave(TArray<UObject*>& OutTexturesToSave) const override;
@@ -311,6 +310,7 @@ public:
 
 	// Begin UMeshTexturePaintingTool Interface.
 	virtual bool AllowsMultiselect() const override { return false; }
+	virtual bool IsMeshAdapterSupported(TSharedPtr<IMeshPaintComponentAdapter> MeshAdapter) const override;
 	virtual UTexture2D* GetSelectedPaintTexture(UMeshComponent const* InMeshComponent) const override;
 	virtual int32 GetSelectedUVChannel(UMeshComponent const* InMeshComponent) const override;
 	virtual void GetModifiedTexturesToSave(TArray<UObject*>& OutTexturesToSave) const override;
diff --git a/Engine/Source/Editor/EditorStyle/Private/StarshipStyle.cpp b/Engine/Source/Editor/EditorStyle/Private/StarshipStyle.cpp
--- a/Engine/Source/Editor/EditorStyle/Private/StarshipStyle.cpp
+++ b/Engine/Source/Editor/EditorStyle/Private/StarshipStyle.cpp
@@ -4843,8 +4843,10 @@ void FStarshipEditorStyle::FStyle::SetupLevelEditorStyle()
 	{
 		Set("MeshPaint.SelectVertex", new IMAGE_BRUSH(TEXT("Icons/GeneralTools/Select_40x"), Icon20x20));
 		Set("MeshPaint.SelectVertex.Small", new IMAGE_BRUSH(TEXT("Icons/GeneralTools/Select_40x"), Icon20x20));
-		Set("MeshPaint.SelectTexture", new IMAGE_BRUSH(TEXT("Icons/GeneralTools/Select_40x"), Icon20x20));
-		Set("MeshPaint.SelectTexture.Small", new IMAGE_BRUSH(TEXT("Icons/GeneralTools/Select_40x"), Icon20x20));
+		Set("MeshPaint.SelectTextureColor", new IMAGE_BRUSH(TEXT("Icons/GeneralTools/Select_40x"), Icon20x20));
+		Set("MeshPaint.SelectTextureColor.Small", new IMAGE_BRUSH(TEXT("Icons/GeneralTools/Select_40x"), Icon20x20));
+		Set("MeshPaint.SelectTextureAsset", new IMAGE_BRUSH(TEXT("Icons/GeneralTools/Select_40x"), Icon20x20));
+		Set("MeshPaint.SelectTextureAsset.Small", new IMAGE_BRUSH(TEXT("Icons/GeneralTools/Select_40x"), Icon20x20));
 		Set("MeshPaint.PaintVertexColor", new IMAGE_BRUSH("Icons/GeneralTools/Paint_40x", Icon20x20));
 		Set("MeshPaint.PaintVertexColor.Small", new IMAGE_BRUSH("Icons/GeneralTools/Paint_40x", Icon20x20));
 		Set("MeshPaint.PaintVertexWeight", new IMAGE_BRUSH("Icons/GeneralTools/Paint_40x", Icon20x20));
