Changed naming for mesh texture paint "Commit" operation.
The operation sets the current painting state to the texture asset.
Apply is a better name than commit here.

diff --git a/Engine/Plugins/MeshPainting/Source/MeshPaintEditorMode/Private/MeshPaintMode.cpp b/Engine/Plugins/MeshPainting/Source/MeshPaintEditorMode/Private/MeshPaintMode.cpp
--- a/Engine/Plugins/MeshPainting/Source/MeshPaintEditorMode/Private/MeshPaintMode.cpp
+++ b/Engine/Plugins/MeshPainting/Source/MeshPaintEditorMode/Private/MeshPaintMode.cpp
@@ -268,13 +268,13 @@ void UMeshPaintMode::BindCommands()
  		FExecuteAction::CreateUObject(this, &UMeshPaintMode::FillTexture),
  		FCanExecuteAction::CreateUObject(this, &UMeshPaintMode::CanFillTexture)));
 
- 	CommandList->MapAction(Commands.Apply, FUIAction(
+ 	CommandList->MapAction(Commands.ApplyVertex, FUIAction(
  		FExecuteAction::CreateUObject(this, &UMeshPaintMode::ApplyVertexColorsToAsset),
  		FCanExecuteAction::CreateUObject(this, &UMeshPaintMode::CanApplyVertexColorsToAsset)));
 
- 	CommandList->MapAction(Commands.Commit, FUIAction(
- 		FExecuteAction::CreateUObject(this, &UMeshPaintMode::CommitTextureColorsToAsset),
- 		FCanExecuteAction::CreateUObject(this, &UMeshPaintMode::CanCommitTextureColorsToAsset)));
+ 	CommandList->MapAction(Commands.ApplyTexture, FUIAction(
+ 		FExecuteAction::CreateUObject(this, &UMeshPaintMode::ApplyTextureColorsToAsset),
+ 		FCanExecuteAction::CreateUObject(this, &UMeshPaintMode::CanApplyTextureColorsToAsset)));
 
 	CommandList->MapAction(Commands.Propagate, FUIAction(
 		FExecuteAction::CreateUObject(this, &UMeshPaintMode::PropagateVertexColorsToLODs),
@@ -907,15 +907,15 @@ bool UMeshPaintMode::CanCycleTextures() const
 	return TexturePaintingTool != nullptr;
 }
 
-void UMeshPaintMode::CommitTextureColorsToAsset()
+void UMeshPaintMode::ApplyTextureColorsToAsset()
 {
 	if (UMeshTexturePaintingTool* TexturePaintingTool = Cast<UMeshTexturePaintingTool>(GetToolManager()->GetActiveTool(EToolSide::Left)))
 	{
-		TexturePaintingTool->CommitAllPaintedTextures();
+		TexturePaintingTool->ApplyAllPaintedTextures();
 	}
 }
 
-bool UMeshPaintMode::CanCommitTextureColorsToAsset() const
+bool UMeshPaintMode::CanApplyTextureColorsToAsset() const
 {
 	if (UMeshTexturePaintingTool* TexturePaintingTool = Cast<UMeshTexturePaintingTool>(GetToolManager()->GetActiveTool(EToolSide::Left)))
 	{
diff --git a/Engine/Plugins/MeshPainting/Source/MeshPaintEditorMode/Private/MeshPaintMode.h b/Engine/Plugins/MeshPainting/Source/MeshPaintEditorMode/Private/MeshPaintMode.h
--- a/Engine/Plugins/MeshPainting/Source/MeshPaintEditorMode/Private/MeshPaintMode.h
+++ b/Engine/Plugins/MeshPainting/Source/MeshPaintEditorMode/Private/MeshPaintMode.h
@@ -100,8 +100,8 @@ protected:
 	bool CanFillTexture() const;
 	void ApplyVertexColorsToAsset();
 	bool CanApplyVertexColorsToAsset() const;
-	void CommitTextureColorsToAsset();
-	bool CanCommitTextureColorsToAsset() const;
+	void ApplyTextureColorsToAsset();
+	bool CanApplyTextureColorsToAsset() const;
 	void PropagateVertexColorsToLODs();
 	bool CanPropagateVertexColorsToLODs() const;
 	void SaveVertexColorsToAssets();
diff --git a/Engine/Plugins/MeshPainting/Source/MeshPaintEditorMode/Private/MeshPaintModeCommands.cpp b/Engine/Plugins/MeshPainting/Source/MeshPaintEditorMode/Private/MeshPaintModeCommands.cpp
--- a/Engine/Plugins/MeshPainting/Source/MeshPaintEditorMode/Private/MeshPaintModeCommands.cpp
+++ b/Engine/Plugins/MeshPainting/Source/MeshPaintEditorMode/Private/MeshPaintModeCommands.cpp
@@ -41,9 +41,9 @@ void FMeshPaintEditorModeCommands::RegisterCommands()
 	UI_COMMAND(SwapColor, "Swap", "Switches the foreground and background colors used for painting", EUserInterfaceActionType::Button, FInputChord(EKeys::X));
 	UI_COMMAND(FillVertex, "Fill", "Fills the selected meshes with the paint color", EUserInterfaceActionType::Button, FInputChord());
 	UI_COMMAND(FillTexture, "Fill", "Fills the selected textures with the paint color", EUserInterfaceActionType::Button, FInputChord());
-	UI_COMMAND(Apply, "Apply", "Propagates instance vertex colors to the source meshes", EUserInterfaceActionType::Button, FInputChord());
-	UI_COMMAND(Commit, "Commit", "Commits texture edits to the source texture assets", EUserInterfaceActionType::Button, FInputChord(EKeys::C, EModifierKey::Control | EModifierKey::Shift));
-	UI_COMMAND(Propagate, "All LODs", "Applies the vertex colors from LOD0 to all LOD levels", EUserInterfaceActionType::Button, FInputChord());
+	UI_COMMAND(ApplyVertex, "Apply", "Applies per instance vertex colors to the source meshes", EUserInterfaceActionType::Button, FInputChord());
+	UI_COMMAND(ApplyTexture, "Apply", "Applies texture edits to the source texture assets", EUserInterfaceActionType::Button, FInputChord(EKeys::C, EModifierKey::Control | EModifierKey::Shift));
+	UI_COMMAND(Propagate, "All LODs", "Propagates the vertex colors from LOD0 to all LOD levels", EUserInterfaceActionType::Button, FInputChord());
 	UI_COMMAND(SaveVertex, "Save", "Saves the source meshes for the selected mesh components", EUserInterfaceActionType::Button, FInputChord());
 	UI_COMMAND(SaveTexture, "Save", "Saves the modified textures for the selected mesh components", EUserInterfaceActionType::Button, FInputChord());
 	UI_COMMAND(Add, "Add", "Adds mesh paint textures to the selected mesh components to enable painting", EUserInterfaceActionType::Button, FInputChord());
@@ -61,22 +61,22 @@ void FMeshPaintEditorModeCommands::RegisterCommands()
 	UI_COMMAND(NextTexture, "Next Texture", "Cycle To next texture", EUserInterfaceActionType::Button, FInputChord(EKeys::Period));
 
 	TArray<TSharedPtr<FUICommandInfo>> VertexColorCommands = {
-		SelectVertex, PaintVertexColor, SwapColor, FillVertex, Apply, Propagate, SaveVertex, RemoveVertex, Copy, Paste, Import, GetTextureColors, Fix
+		SelectVertex, PaintVertexColor, SwapColor, FillVertex, Propagate, ApplyVertex, SaveVertex, RemoveVertex, Copy, Paste, Import, GetTextureColors, Fix
 	};
 	Commands.Add(UMeshPaintMode::MeshPaintMode_VertexColor, VertexColorCommands);
 
 	TArray<TSharedPtr<FUICommandInfo>> VertexWeightCommands = {
-		SelectVertex, PaintVertexWeight, SwapColor, FillVertex, Apply, Propagate, SaveVertex, RemoveVertex, Copy, Paste, Import, Fix
+		SelectVertex, PaintVertexWeight, SwapColor, FillVertex, Propagate, ApplyVertex, SaveVertex, RemoveVertex, Copy, Paste, Import, Fix
 	};
 	Commands.Add(UMeshPaintMode::MeshPaintMode_VertexWeights, VertexWeightCommands);
 
 	TArray<TSharedPtr<FUICommandInfo>> TextureColorCommands = {
-		SelectTextureColor, PaintTextureColor, SwapColor, FillTexture, Commit, SaveTexture, Add, RemoveTexture, Copy, Paste, GetVertexColors
+		SelectTextureColor, PaintTextureColor, SwapColor, FillTexture, ApplyTexture, SaveTexture, Add, RemoveTexture, Copy, Paste, GetVertexColors
 	};
 	Commands.Add(UMeshPaintMode::MeshPaintMode_TextureColor, TextureColorCommands);
 
 	TArray<TSharedPtr<FUICommandInfo>> TextureAssetCommands = {
-		SelectTextureAsset, PaintTextureAsset, SwapColor, FillTexture, Commit, SaveTexture
+		SelectTextureAsset, PaintTextureAsset, SwapColor, FillTexture, ApplyTexture, SaveTexture
 	};
 	Commands.Add(UMeshPaintMode::MeshPaintMode_TextureAsset, TextureAssetCommands);
 }
diff --git a/Engine/Plugins/MeshPainting/Source/MeshPaintEditorMode/Public/MeshPaintModeCommands.h b/Engine/Plugins/MeshPainting/Source/MeshPaintEditorMode/Public/MeshPaintModeCommands.h
--- a/Engine/Plugins/MeshPainting/Source/MeshPaintEditorMode/Public/MeshPaintModeCommands.h
+++ b/Engine/Plugins/MeshPainting/Source/MeshPaintEditorMode/Public/MeshPaintModeCommands.h
@@ -42,8 +42,8 @@ public:
 	TSharedPtr<FUICommandInfo> SwapColor;
 	TSharedPtr<FUICommandInfo> FillVertex;
 	TSharedPtr<FUICommandInfo> FillTexture;
-	TSharedPtr<FUICommandInfo> Apply;
-	TSharedPtr<FUICommandInfo> Commit;
+	TSharedPtr<FUICommandInfo> ApplyVertex;
+	TSharedPtr<FUICommandInfo> ApplyTexture;
 	TSharedPtr<FUICommandInfo> Propagate;
 	TSharedPtr<FUICommandInfo> SaveVertex;
 	TSharedPtr<FUICommandInfo> SaveTexture;
diff --git a/Engine/Plugins/MeshPainting/Source/MeshPaintingToolset/Private/MeshTexturePaintingTool.cpp b/Engine/Plugins/MeshPainting/Source/MeshPaintingToolset/Private/MeshTexturePaintingTool.cpp
--- a/Engine/Plugins/MeshPainting/Source/MeshPaintingToolset/Private/MeshTexturePaintingTool.cpp
+++ b/Engine/Plugins/MeshPainting/Source/MeshPaintingToolset/Private/MeshTexturePaintingTool.cpp
@@ -720,7 +720,7 @@ FInputRayHit UMeshTexturePaintingTool::CanBeginClickDragSequence(const FInputDev
 	{
 		MeshPaintingSubsystem->LastPaintedComponent = (UMeshComponent*)LastBestHitResult.Component.Get();
 		GVertexViewModeOverrideOwnerName = *LastBestHitResult.Component->GetOwner()->GetName();
-		//CommitAllPaintedTextures();
+		//ApplyAllPaintedTextures();
 		//MeshPaintingSubsystem->bNeedsRecache = true;
 	}
 
@@ -1397,7 +1397,7 @@ void UMeshTexturePaintingTool::FinishPaintingTexture()
 		FPaintTexture2DData* TextureData = GetPaintTargetData(PaintingTexture2D);
 		check(TextureData);
 
-		// Commit to the texture source art but don't do any compression, compression is saved for the CommitAllPaintedTextures function.
+		// Set to the texture source art but don't do any compression, compression is saved for the ApplyAllPaintedTextures function.
 		if (TextureData->bIsPaintingTexture2DModified == true)
 		{
 			const int32 TexWidth = TextureData->PaintRenderTargetTexture->SizeX;
@@ -1593,16 +1593,16 @@ void UMeshTextureAssetPaintingTool::CycleTextures(int32 Direction)
 }
 
 
-void UMeshTexturePaintingTool::CommitAllPaintedTextures()
+void UMeshTexturePaintingTool::ApplyAllPaintedTextures()
 {
 	if (PaintTargetData.Num() > 0)
 	{
 		check(PaintingTexture2D == nullptr);
 
-		FScopedTransaction Transaction(LOCTEXT("MeshPaintMode_TexturePaint_Commit", "Commit Texture Paint"));
+		FScopedTransaction Transaction(LOCTEXT("MeshPaintMode_TexturePaint_Apply", "Apply Texture Paint"));
 		
 		Modify();
-	//	GWarn->BeginSlowTask(LOCTEXT("BeginMeshPaintMode_TexturePaint_CommitTask", "Committing Texture Paint Changes"), true);
+	//	GWarn->BeginSlowTask(LOCTEXT("BeginMeshPaintMode_TexturePaint_ApplyTask", "Applying Texture Paint Changes"), true);
 
 		int32 CurStep = 1;
 		int32 TotalSteps = GetNumberOfPendingPaintChanges();
@@ -1611,10 +1611,10 @@ void UMeshTexturePaintingTool::CommitAllPaintedTextures()
 		{
 			FPaintTexture2DData* TextureData = &It.Value();
 
-			// Commit the texture
+			// Apply the texture
 			if (TextureData->bIsPaintingTexture2DModified == true)
 			{
-			//	GWarn->StatusUpdate(CurStep++, TotalSteps, FText::Format(LOCTEXT("MeshPaintMode_TexturePaint_CommitStatus", "Committing Texture Paint Changes: {0}"), FText::FromName(TextureData->PaintingTexture2D->GetFName())));
+			//	GWarn->StatusUpdate(CurStep++, TotalSteps, FText::Format(LOCTEXT("MeshPaintMode_TexturePaint_ApplyStatus", "Applying Texture Paint Changes: {0}"), FText::FromName(TextureData->PaintingTexture2D->GetFName())));
 
 				const int32 TexWidth = TextureData->PaintRenderTargetTexture->SizeX;
 				const int32 TexHeight = TextureData->PaintRenderTargetTexture->SizeY;
@@ -1694,7 +1694,7 @@ int32 UMeshTexturePaintingTool::GetNumberOfPendingPaintChanges() const
 	{
 		const FPaintTexture2DData* TextureData = &It.Value();
 
-		// Commit the texture
+		// Apply the texture
 		if (TextureData->bIsPaintingTexture2DModified == true)
 		{
 			Result++;
diff --git a/Engine/Plugins/MeshPainting/Source/MeshPaintingToolset/Public/MeshTexturePaintingTool.h b/Engine/Plugins/MeshPainting/Source/MeshPaintingToolset/Public/MeshTexturePaintingTool.h
--- a/Engine/Plugins/MeshPainting/Source/MeshPaintingToolset/Public/MeshTexturePaintingTool.h
+++ b/Engine/Plugins/MeshPainting/Source/MeshPaintingToolset/Public/MeshTexturePaintingTool.h
@@ -161,7 +161,7 @@ public:
 		return OnPaintingFinishedDelegate;
 	}
 
-	void CommitAllPaintedTextures();
+	void ApplyAllPaintedTextures();
 	void ClearAllTextureOverrides();
 	
 	/** Returns the number of texture that require a commit. */
diff --git a/Engine/Source/Editor/EditorStyle/Private/StarshipStyle.cpp b/Engine/Source/Editor/EditorStyle/Private/StarshipStyle.cpp
--- a/Engine/Source/Editor/EditorStyle/Private/StarshipStyle.cpp
+++ b/Engine/Source/Editor/EditorStyle/Private/StarshipStyle.cpp
@@ -4862,10 +4862,10 @@ void FStarshipEditorStyle::FStyle::SetupLevelEditorStyle()
 		Set("MeshPaint.FillVertex.Small", new IMAGE_BRUSH("/Icons/GeneralTools/PaintBucket_40x", Icon20x20));
 		Set("MeshPaint.FillTexture", new IMAGE_BRUSH("/Icons/GeneralTools/PaintBucket_40x", Icon20x20));
 		Set("MeshPaint.FillTexture.Small", new IMAGE_BRUSH("/Icons/GeneralTools/PaintBucket_40x", Icon20x20));
-		Set("MeshPaint.Apply", new IMAGE_BRUSH("/Icons/Paint/Paint_Propagate_40x", Icon20x20));
-		Set("MeshPaint.Apply.Small", new IMAGE_BRUSH("/Icons/Paint/Paint_Propagate_40x", Icon20x20));
-		Set("MeshPaint.Commit", new IMAGE_BRUSH("/Icons/Paint/Paint_PropagateTexture_40x", Icon20x20));
-		Set("MeshPaint.Commit.Small", new IMAGE_BRUSH("/Icons/Paint/Paint_PropagateTexture_40x", Icon20x20));
+		Set("MeshPaint.ApplyVertex", new IMAGE_BRUSH("/Icons/Paint/Paint_Propagate_40x", Icon20x20));
+		Set("MeshPaint.ApplyVertex.Small", new IMAGE_BRUSH("/Icons/Paint/Paint_Propagate_40x", Icon20x20));
+		Set("MeshPaint.ApplyTexture", new IMAGE_BRUSH("/Icons/Paint/Paint_PropagateTexture_40x", Icon20x20));
+		Set("MeshPaint.ApplyTexture.Small", new IMAGE_BRUSH("/Icons/Paint/Paint_PropagateTexture_40x", Icon20x20));
 		Set("MeshPaint.Propagate", new IMAGE_BRUSH("Icons/Paint/Paint_AllLODs_40x", Icon20x20));
 		Set("MeshPaint.Propagate.Small", new IMAGE_BRUSH("Icons/Paint/Paint_AllLODs_40x", Icon20x20));
 		Set("MeshPaint.SaveVertex", new IMAGE_BRUSH("/Icons/GeneralTools/Save_40x", Icon20x20));
