Improve mesh paint tool selection flow.
Now we can freely select actors while in the tool, and the paint option will enable or disable depending on the tool and selection state.
Also when in a paint subtool, we continuously detect when the current selection is no longer valid and exit the tool back to the selection subtool.
The validity detection supports selection, selected component settings, selected materials state, nanite support etc.

diff --git a/Engine/Plugins/MeshPainting/Source/MeshPaintEditorMode/Private/MeshPaintMode.cpp b/Engine/Plugins/MeshPainting/Source/MeshPaintEditorMode/Private/MeshPaintMode.cpp
--- a/Engine/Plugins/MeshPainting/Source/MeshPaintEditorMode/Private/MeshPaintMode.cpp
+++ b/Engine/Plugins/MeshPainting/Source/MeshPaintEditorMode/Private/MeshPaintMode.cpp
@@ -161,6 +161,9 @@ void UMeshPaintMode::Enter()
 
 	FLevelEditorModule& LevelEditor = FModuleManager::GetModuleChecked<FLevelEditorModule>(FName(TEXT("LevelEditor")));
 	LevelEditor.OnRedrawLevelEditingViewports().AddUObject(this, &UMeshPaintMode::UpdateOnMaterialChange);
+
+	// some global cvars can affect whether painting is valid (nanite on/off etc)
+	CVarDelegateHandle = IConsoleManager::Get().RegisterConsoleVariableSink_Handle(FConsoleCommandDelegate::CreateLambda([this]{ bRecacheValidForPaint = true; }));
 }
 
 void UMeshPaintMode::Exit()
@@ -186,6 +189,9 @@ void UMeshPaintMode::Exit()
 
 	FLevelEditorModule& LevelEditor = FModuleManager::GetModuleChecked<FLevelEditorModule>(FName(TEXT("LevelEditor")));
 	LevelEditor.OnRedrawLevelEditingViewports().RemoveAll(this);
+
+	IConsoleManager::Get().UnregisterConsoleVariableSink_Handle(CVarDelegateHandle);
+	CVarDelegateHandle = {};
 }
 
 void UMeshPaintMode::CreateToolkit()
@@ -195,20 +201,30 @@ void UMeshPaintMode::CreateToolkit()
 
 void UMeshPaintMode::Tick(FEditorViewportClient* ViewportClient, float DeltaTime)
 {
-	EMeshPaintActiveMode CurrentActiveMode = EMeshPaintActiveMode::VertexColor;
-
 	if (bRecacheVertexDataSize)
 	{
 		UpdateCachedVertexDataSize();
 	}
 
-	// Make sure that correct tab is visible for the current tool
-	// Note that currently Color and Weight mode share the same Select tool
+	if (bRecacheValidForPaint)
+	{
+		GEngine->GetEngineSubsystem<UMeshPaintingSubsystem>()->UpdatePaintSupportState();
+		bRecacheValidForPaint = false;
+	}
+
+	// Close the active paint tool if selection (or other state) changes mean that it's not longer valid to paint.
+	// For example if the selected component or it's materials no longer supports texture painting.
+	EndPaintToolIfNoLongerValid();
+
+	// Make sure that correct tab is visible for the current tool.
+	// Note that currently Color and Weight mode share the same Select tool.
 	UInteractiveTool const* ActiveTool = GetToolManager()->GetActiveTool(EToolSide::Mouse);
 	const FString ActiveToolName = GetToolManager()->GetActiveToolName(EToolSide::Mouse);
 
 	FName ActiveTab = Toolkit->GetCurrentPalette();
 	FName TargetTab = ActiveTab;
+	EMeshPaintActiveMode CurrentActiveMode = EMeshPaintActiveMode::VertexColor;
+	
 	if (ActiveToolName == VertexColorPaintToolName)
 	{
 		TargetTab = MeshPaintMode_VertexColor;
@@ -384,6 +400,12 @@ void UMeshPaintMode::ActorSelectionChangeNotify()
 	UpdateSelectedMeshes();
 }
 
+void UMeshPaintMode::ActorPropChangeNotify()
+{
+	// Setting change on selected components can change whether they are valid for painting.
+	bRecacheValidForPaint = true;
+}
+
 void UMeshPaintMode::UpdateSelectedMeshes()
 {
 	if (UMeshPaintingSubsystem* MeshPaintingSubsystem = GEngine->GetEngineSubsystem<UMeshPaintingSubsystem>())
@@ -392,90 +414,47 @@ void UMeshPaintMode::UpdateSelectedMeshes()
 		TArray<UMeshComponent*> CurrentMeshComponents = GetSelectedComponents<UMeshComponent>();
 		MeshPaintingSubsystem->AddSelectedMeshComponents(CurrentMeshComponents);
 		MeshPaintingSubsystem->bNeedsRecache = true;
-
-		// Check the current selection for a material that supports texture paint
-		FName PaletteName = Toolkit->GetCurrentPalette();
-		if (PaletteName == UMeshPaintMode::MeshPaintMode_TextureColor || PaletteName == UMeshPaintMode::MeshPaintMode_TextureAsset)
-		{
-			UpdateToolForSelection(CurrentMeshComponents);
-		}
 	}
 	
 	bRecacheVertexDataSize = true;
+	bRecacheValidForPaint = true;
 }
 
-void UMeshPaintMode::UpdateToolForSelection(const TArray<UMeshComponent*>& CurrentMeshComponents)
+void UMeshPaintMode::EndPaintToolIfNoLongerValid()
 {
-	FName PaletteName = Toolkit->GetCurrentPalette();
-	const bool bIsTextureColorPainting = PaletteName == UMeshPaintMode::MeshPaintMode_TextureColor;
+	bool bInvalidTool = false;
 
-	bool bCurrentSelectionSupportsTexturePaint = false;
-	
-	// Texture Color painting supports multiselect. But other modes require a single mesh selection.
-	if (CurrentMeshComponents.Num() == 1 || bIsTextureColorPainting)
+	UInteractiveToolManager* ToolManager = GetToolManager();
+	UInteractiveTool const* Tool = (ToolManager == nullptr) ? nullptr : ToolManager->GetActiveTool(EToolSide::Mouse);
+	if (Tool != nullptr)
 	{
-		TArray<FPaintableTexture> PaintableTextures;
+		UMeshPaintingSubsystem* MeshPaintingSubsystem = GEngine->GetEngineSubsystem<UMeshPaintingSubsystem>();
 
-		for (UMeshComponent* MeshComponent : CurrentMeshComponents)
+		if (Tool->IsA<UMeshVertexPaintingTool>())
 		{
-			if (MeshComponent)
-			{
-				TSharedPtr<IMeshPaintComponentAdapter> MeshAdapter = GEngine->GetEngineSubsystem<UMeshPaintingSubsystem>()->GetAdapterForComponent(MeshComponent);
-				if (!MeshAdapter.IsValid())
-				{
-					MeshAdapter = FMeshPaintComponentAdapterFactory::CreateAdapterForMesh(MeshComponent, 0);
-					if (MeshAdapter)
-					{
-						GEngine->GetEngineSubsystem<UMeshPaintingSubsystem>()->AddToComponentToAdapterMap(MeshComponent, MeshAdapter);
-					}
-				}
-
-				int32 DummyDefaultIndex = INDEX_NONE;
-				UTexturePaintToolset::RetrieveTexturesForComponent(MeshComponent, MeshAdapter.Get(), DummyDefaultIndex, PaintableTextures);
-			}
+			bInvalidTool = !MeshPaintingSubsystem->GetSelectionSupportsVertexPaint();
+		}
+		else if (Tool->IsA<UMeshTextureColorPaintingTool>())
+		{
+			bInvalidTool = !MeshPaintingSubsystem->GetSelectionSupportsTextureColorPaint();
+		}
+		else if (Tool->IsA<UMeshTextureAssetPaintingTool>())
+		{
+			bInvalidTool = !MeshPaintingSubsystem->GetSelectionSupportsTextureAssetPaint();
 		}
-
-		// PaintableTextures are collected for both TextureColor painting (MeshPaintTextures on components) and Texture painting (Textures ref'd in materials).
-		// Filter these according to the current texture painting mode here.
- 		PaintableTextures.RemoveAll([bIsTextureColorPainting](FPaintableTexture const& PaintableTexture) { return bIsTextureColorPainting != PaintableTexture.bIsMeshTexture; });
-
-		bCurrentSelectionSupportsTexturePaint = PaintableTextures.Num() > 0;
 	}
 
-	GEngine->GetEngineSubsystem<UMeshPaintingSubsystem>()->SetSelectionHasMaterialValidForTexturePaint(bCurrentSelectionSupportsTexturePaint);
-
-	if (!bCurrentSelectionSupportsTexturePaint)
+	if (bInvalidTool)
 	{
-		UInteractiveToolManager* ToolManager = GetToolManager();
-		UInteractiveTool const* Tool = (ToolManager == nullptr) ? nullptr : ToolManager->GetActiveTool(EToolSide::Mouse);
-		if (Tool != nullptr)
-		{
-			if (Tool->IsA<UMeshTexturePaintingTool>())
-			{
-				GetInteractiveToolsContext()->EndTool(EToolShutdownType::Accept);
-
-				if (Tool->IsA<UMeshTextureColorPaintingTool>())
-				{
-					GetInteractiveToolsContext()->StartTool(TextureColorSelectToolName);
-				}
-				else if (Tool->IsA<UMeshTextureAssetPaintingTool>())
-				{
-					GetInteractiveToolsContext()->StartTool(TextureAssetSelectToolName);
-				}
-			}
-		}
+		GetInteractiveToolsContext()->EndTool(EToolShutdownType::Accept);
+		ActivateDefaultTool();
 	}
 }
 
 void UMeshPaintMode::UpdateOnMaterialChange(bool bInvalidateHitProxies)
 {
-	// Check the current selection for a material that supports texture paint
-	FName PaletteName = Toolkit->GetCurrentPalette();
-	if (PaletteName == UMeshPaintMode::MeshPaintMode_TextureColor || PaletteName == UMeshPaintMode::MeshPaintMode_TextureAsset)
-	{
-		const TArray<UMeshComponent*> CurrentMeshComponents = GetSelectedComponents<UMeshComponent>();
-		UpdateToolForSelection(CurrentMeshComponents);
-	}
+	// Need to recheck whether the current material supports texture paint.
+	bRecacheValidForPaint = true;
 }
 
 void UMeshPaintMode::OnObjectsReplaced(const TMap<UObject*, UObject*>& InOldToNewInstanceMap)
@@ -1073,9 +1052,8 @@ void UMeshPaintMode::AddMeshPaintTextures()
 		GEngine->GetEngineSubsystem<UMeshPaintingSubsystem>()->CreateComponentMeshPaintTexture(Component);
 	}
 
-	// The selection/paint status will change after add.
-	const TArray<UMeshComponent*> CurrentMeshComponents = GetSelectedComponents<UMeshComponent>();
-	UpdateToolForSelection(CurrentMeshComponents);
+	// The paint status will change after add.
+	bRecacheValidForPaint = true;
 }
 
 bool UMeshPaintMode::CanAddMeshPaintTextures() const
@@ -1103,9 +1081,8 @@ void UMeshPaintMode::RemoveMeshPaintTexture()
 		GEngine->GetEngineSubsystem<UMeshPaintingSubsystem>()->RemoveComponentMeshPaintTexture(Component);
 	}
 
-	// The selection/paint status will change after removal.
-	const TArray<UMeshComponent*> CurrentMeshComponents = GetSelectedComponents<UMeshComponent>();
-	UpdateToolForSelection(CurrentMeshComponents);
+	// The paint status will change after removal.
+	bRecacheValidForPaint = true;
 }
 
 bool UMeshPaintMode::CanRemoveMeshPaintTextures() const
@@ -1184,8 +1161,7 @@ void UMeshPaintMode::PasteMeshPaintTexture()
 	}
 
 	// If this adds a mesh paint texture the selection/paint status may have changed.
-	const TArray<UMeshComponent*> CurrentMeshComponents = GetSelectedComponents<UMeshComponent>();
-	UpdateToolForSelection(CurrentMeshComponents);
+	bRecacheValidForPaint = true;
 }
 
 bool UMeshPaintMode::CanPasteMeshPaintTexture() const
@@ -1273,8 +1249,7 @@ void UMeshPaintMode::ImportMeshPaintTextureFromVertexColors()
 	}
 
 	// If this adds a mesh paint texture the selection/paint status may have changed.
-	const TArray<UMeshComponent*> CurrentMeshComponents = GetSelectedComponents<UMeshComponent>();
-	UpdateToolForSelection(CurrentMeshComponents);
+	bRecacheValidForPaint = true;
 }
 
 bool UMeshPaintMode::CanImportMeshPaintTextureFromVertexColors() const
diff --git a/Engine/Plugins/MeshPainting/Source/MeshPaintEditorMode/Private/MeshPaintMode.h b/Engine/Plugins/MeshPainting/Source/MeshPaintEditorMode/Private/MeshPaintMode.h
--- a/Engine/Plugins/MeshPainting/Source/MeshPaintEditorMode/Private/MeshPaintMode.h
+++ b/Engine/Plugins/MeshPainting/Source/MeshPaintEditorMode/Private/MeshPaintMode.h
@@ -73,20 +73,20 @@ protected:
 	virtual void OnToolStarted(UInteractiveToolManager* Manager, UInteractiveTool* Tool) override;
 	virtual void OnToolEnded(UInteractiveToolManager* Manager, UInteractiveTool* Tool) override;
 	virtual void ActorSelectionChangeNotify() override;
+	virtual void ActorPropChangeNotify() override;
 	virtual void ActivateDefaultTool() override;
 	virtual void UpdateOnPaletteChange(FName NewPalette);
 	// end UEdMode Interface
 	
-	void UpdateSelectedMeshes();
-
-	void UpdateToolForSelection(const TArray<UMeshComponent*>& CurrentMeshComponents);
 
+	void UpdateSelectedMeshes();
 	void UpdateOnMaterialChange(bool bInvalidateHitProxies);
 	void OnObjectsReplaced(const TMap<UObject*, UObject*>& InOldToNewInstanceMap);
 	void OnResetViewMode();
 	void OnVertexPaintFinished();
 
 	void UpdateCachedVertexDataSize();
+	void EndPaintToolIfNoLongerValid();
 
 	bool IsInSelectTool() const;
 	bool IsInPaintTool() const;
@@ -144,11 +144,11 @@ protected:
 	UPROPERTY(Transient)
 	TObjectPtr<UMeshPaintModeSettings> ModeSettings;
 
-	// End vertex paint state
-	FGetSelectedMeshComponents MeshComponentDelegate;
-	uint32 CachedVertexDataSize;
-	bool bRecacheVertexDataSize;
+	uint32 CachedVertexDataSize = 0;
+	bool bRecacheVertexDataSize = false;
+	
+	bool bRecacheValidForPaint = false;
 
 	FDelegateHandle PaletteChangedHandle;
+	FConsoleVariableSinkHandle CVarDelegateHandle;
 };
-
diff --git a/Engine/Plugins/MeshPainting/Source/MeshPaintingToolset/Private/MeshPaintHelpers.cpp b/Engine/Plugins/MeshPainting/Source/MeshPaintingToolset/Private/MeshPaintHelpers.cpp
--- a/Engine/Plugins/MeshPainting/Source/MeshPaintingToolset/Private/MeshPaintHelpers.cpp
+++ b/Engine/Plugins/MeshPainting/Source/MeshPaintingToolset/Private/MeshPaintHelpers.cpp
@@ -1,20 +1,22 @@
 // Copyright Epic Games, Inc. All Rights Reserved.
 
 #include "MeshPaintHelpers.h"
+
 #include "ComponentReregisterContext.h"
-#include "IMeshPaintComponentAdapter.h"
-#include "MeshPaintAdapterFactory.h"
-#include "Components/StaticMeshComponent.h"
 #include "Components/SkeletalMeshComponent.h"
+#include "Components/StaticMeshComponent.h"
 #include "Engine/SkeletalMesh.h"
 #include "Engine/SkinnedAssetCommon.h"
 #include "Engine/Texture2D.h"
+#include "IMeshPaintComponentAdapter.h"
+#include "MeshPaintAdapterFactory.h"
+#include "MeshVertexPaintingTool.h"
 #include "Rendering/SkeletalMeshLODModel.h"
-#include "StaticMeshComponentLODInfo.h"
-#include "StaticMeshAttributes.h"
-#include "Rendering/SkeletalMeshRenderData.h"
 #include "Rendering/SkeletalMeshModel.h"
-#include "MeshVertexPaintingTool.h"
+#include "Rendering/SkeletalMeshRenderData.h"
+#include "StaticMeshAttributes.h"
+#include "StaticMeshComponentLODInfo.h"
+#include "TexturePaintToolset.h"
 #include "VT/MeshPaintVirtualTexture.h"
 
 #include UE_INLINE_GENERATED_CPP_BY_NAME(MeshPaintHelpers)
@@ -22,8 +24,10 @@
 extern void PropagateVertexPaintToSkeletalMesh(USkeletalMesh* SkeletalMesh, int32 LODIndex);
 
 UMeshPaintingSubsystem::UMeshPaintingSubsystem()
-	:	bNeedsRecache(true),
-		bSelectionHasMaterialValidForTexturePaint(false)
+	: bNeedsRecache(true)
+	, bSelectionSupportsVertexPaint(false)
+	, bSelectionSupportsTextureColorPaint(false)
+	, bSelectionSupportsTextureAssetPaint(false)
 {
 
 }
@@ -1570,17 +1574,53 @@ bool UMeshPaintingSubsystem::FindHitResult(const FRay Ray, FHitResult& BestTrace
 	return BestTraceResult.Distance != FLT_MAX;
 }
 
-bool UMeshPaintingSubsystem::SelectionContainsValidAdapters() const
+void UMeshPaintingSubsystem::UpdatePaintSupportState()
 {
-	for (const auto& MeshAdapterPair : ComponentToAdapterMap)
+	bSelectionSupportsVertexPaint = false;
+	bSelectionSupportsTextureColorPaint = false;
+	bSelectionSupportsTextureAssetPaint = false;
+
+	for (TWeakObjectPtr<UMeshComponent> MeshComponentWeak : SelectedMeshComponents)
 	{
-		if (MeshAdapterPair.Value && MeshAdapterPair.Value->IsValid())
+		if (UMeshComponent* MeshComponent = MeshComponentWeak.Get())
 		{
-			return true;
+			TSharedPtr<IMeshPaintComponentAdapter> MeshAdapter = GetAdapterForComponent(MeshComponent);
+			if (!MeshAdapter.IsValid())
+			{
+				MeshAdapter = FMeshPaintComponentAdapterFactory::CreateAdapterForMesh(MeshComponent, 0);
+				if (MeshAdapter)
+				{
+					AddToComponentToAdapterMap(MeshComponent, MeshAdapter);
+				}
+			}
+
+			if (MeshAdapter)
+			{
+				bSelectionSupportsVertexPaint |= MeshAdapter->SupportsVertexPaint();
+
+				const bool bSupportsTextureColorPaint = MeshAdapter->SupportsTextureColorPaint();
+				const bool bSupportsTextureAssetPaint = MeshAdapter->SupportsTexturePaint();
+
+				if (bSupportsTextureColorPaint || bSupportsTextureAssetPaint)
+				{
+					// Collect PaintableTextures. This if for both TextureColor painting (MeshPaintTextures on components) and TextureAsset painting (Textures ref'd in materials).
+					int32 DummyDefaultIndex = INDEX_NONE;
+					TArray<FPaintableTexture> PaintableTextures;
+					UTexturePaintToolset::RetrieveTexturesForComponent(MeshComponent, MeshAdapter.Get(), DummyDefaultIndex, PaintableTextures);
+
+					for (FPaintableTexture const& PaintableTexture : PaintableTextures)
+					{
+						// The bIsMeshTexture tells us which mode the paintable texture works with.
+						bSelectionSupportsTextureColorPaint |= bSupportsTextureColorPaint && PaintableTexture.bIsMeshTexture;
+						bSelectionSupportsTextureAssetPaint |= bSupportsTextureAssetPaint && !PaintableTexture.bIsMeshTexture;
+					}
+				}
+			}
 		}
 	}
 
-	return false;
+	// Texture Asset painting only supports single component select.
+	bSelectionSupportsTextureAssetPaint &= SelectedMeshComponents.Num() == 1;
 }
 
 TArray<FPerComponentVertexColorData> UMeshPaintingSubsystem::GetCopiedColorsByComponent() const
diff --git a/Engine/Plugins/MeshPainting/Source/MeshPaintingToolset/Private/MeshSelect.cpp b/Engine/Plugins/MeshPainting/Source/MeshPaintingToolset/Private/MeshSelect.cpp
--- a/Engine/Plugins/MeshPainting/Source/MeshPaintingToolset/Private/MeshSelect.cpp
+++ b/Engine/Plugins/MeshPainting/Source/MeshPaintingToolset/Private/MeshSelect.cpp
@@ -87,20 +87,9 @@ void UMeshClickTool::OnClicked(const FInputDeviceRay& ClickPos)
 	SelectionMechanic->OnClicked(ClickPos);
 }
 
-
-bool UVertexAdapterClickTool::IsMeshAdapterSupported(TSharedPtr<IMeshPaintComponentAdapter> MeshAdapter) const
-{
-	return MeshAdapter.IsValid() ? MeshAdapter->SupportsVertexPaint() : false;
-}
-
-bool UTextureColorAdapterClickTool::IsMeshAdapterSupported(TSharedPtr<IMeshPaintComponentAdapter> MeshAdapter) const
-{
-	return MeshAdapter.IsValid() ? MeshAdapter->SupportsTextureColorPaint() : false;
-}
-
-bool UTextureAssetAdapterClickTool::IsMeshAdapterSupported(TSharedPtr<IMeshPaintComponentAdapter> MeshAdapter) const
+bool UMeshClickTool::IsMeshAdapterSupported(TSharedPtr<IMeshPaintComponentAdapter> MeshAdapter) const
 {
-	return MeshAdapter.IsValid() ? MeshAdapter->SupportsTexturePaint() : false;
+	return MeshAdapter.IsValid();
 }
 
 #undef LOCTEXT_NAMESPACE
diff --git a/Engine/Plugins/MeshPainting/Source/MeshPaintingToolset/Private/MeshTexturePaintingTool.cpp b/Engine/Plugins/MeshPainting/Source/MeshPaintingToolset/Private/MeshTexturePaintingTool.cpp
--- a/Engine/Plugins/MeshPainting/Source/MeshPaintingToolset/Private/MeshTexturePaintingTool.cpp
+++ b/Engine/Plugins/MeshPainting/Source/MeshPaintingToolset/Private/MeshTexturePaintingTool.cpp
@@ -140,7 +140,7 @@ namespace UE::MeshPaintingToolset::Private::TexturePainting
 
 bool UMeshTextureColorPaintingToolBuilder::CanBuildTool(const FToolBuilderState& SceneState) const
 {
-	return GEngine->GetEngineSubsystem<UMeshPaintingSubsystem>()->SelectionHasMaterialValidForTexturePaint();
+	return GEngine->GetEngineSubsystem<UMeshPaintingSubsystem>()->GetSelectionSupportsTextureColorPaint();
 }
 
 UInteractiveTool* UMeshTextureColorPaintingToolBuilder::BuildTool(const FToolBuilderState& SceneState) const
@@ -150,7 +150,7 @@ UInteractiveTool* UMeshTextureColorPaintingToolBuilder::BuildTool(const FToolBui
 
 bool UMeshTextureAssetPaintingToolBuilder::CanBuildTool(const FToolBuilderState& SceneState) const
 {
-	return GEngine->GetEngineSubsystem<UMeshPaintingSubsystem>()->SelectionHasMaterialValidForTexturePaint();
+	return GEngine->GetEngineSubsystem<UMeshPaintingSubsystem>()->GetSelectionSupportsTextureAssetPaint();
 }
 
 UInteractiveTool* UMeshTextureAssetPaintingToolBuilder::BuildTool(const FToolBuilderState& SceneState) const
diff --git a/Engine/Plugins/MeshPainting/Source/MeshPaintingToolset/Private/MeshVertexPaintingTool.cpp b/Engine/Plugins/MeshPainting/Source/MeshPaintingToolset/Private/MeshVertexPaintingTool.cpp
--- a/Engine/Plugins/MeshPainting/Source/MeshPaintingToolset/Private/MeshVertexPaintingTool.cpp
+++ b/Engine/Plugins/MeshPainting/Source/MeshPaintingToolset/Private/MeshVertexPaintingTool.cpp
@@ -20,7 +20,7 @@
 
 bool UMeshColorPaintingToolBuilder::CanBuildTool(const FToolBuilderState& SceneState) const
 {
-	return true;
+	return GEngine->GetEngineSubsystem<UMeshPaintingSubsystem>()->GetSelectionSupportsVertexPaint();
 }
 
 UInteractiveTool* UMeshColorPaintingToolBuilder::BuildTool(const FToolBuilderState& SceneState) const
@@ -31,7 +31,7 @@ UInteractiveTool* UMeshColorPaintingToolBuilder::BuildTool(const FToolBuilderSta
 
 bool UMeshWeightPaintingToolBuilder::CanBuildTool(const FToolBuilderState& SceneState) const
 {
-	return true;
+	return GEngine->GetEngineSubsystem<UMeshPaintingSubsystem>()->GetSelectionSupportsVertexPaint();
 }
 
 UInteractiveTool* UMeshWeightPaintingToolBuilder::BuildTool(const FToolBuilderState& SceneState) const
diff --git a/Engine/Plugins/MeshPainting/Source/MeshPaintingToolset/Public/MeshPaintHelpers.h b/Engine/Plugins/MeshPainting/Source/MeshPaintingToolset/Public/MeshPaintHelpers.h
--- a/Engine/Plugins/MeshPainting/Source/MeshPaintingToolset/Public/MeshPaintHelpers.h
+++ b/Engine/Plugins/MeshPainting/Source/MeshPaintingToolset/Public/MeshPaintHelpers.h
@@ -238,16 +238,6 @@ public:
 	/** Helper function to retrieve vertex color from a UTexture given a UVCoordinate */
 	FColor PickVertexColorFromTextureData(const uint8* MipData, const FVector2D& UVCoordinate, const UTexture2D* Texture, const FColor ColorMask);	
 
-	void SetSelectionHasMaterialValidForTexturePaint(const bool InValidity)
-	{
-		bSelectionHasMaterialValidForTexturePaint = InValidity;
-	}
-
-	bool SelectionHasMaterialValidForTexturePaint()
-	{
-		return bSelectionHasMaterialValidForTexturePaint;
-	}
-
 	/** Map of geometry adapters for each selected mesh component */
 	TSharedPtr<IMeshPaintComponentAdapter> GetAdapterForComponent(const UMeshComponent* InComponent) const;
 	void AddToComponentToAdapterMap(const UMeshComponent* InComponent, const TSharedPtr<IMeshPaintComponentAdapter> InAdapter);
@@ -259,7 +249,6 @@ public:
 	TArray<UMeshComponent*> GetPaintableMeshComponents() const;
 	void AddPaintableMeshComponent(UMeshComponent* InComponent);
 	void ClearPaintableMeshComponents();
-	bool SelectionContainsValidAdapters() const;
 	TArray<FPerComponentVertexColorData> GetCopiedColorsByComponent() const;
 	void SetCopiedColorsByComponent(TArray<FPerComponentVertexColorData>& InCopiedColors);
 	void CacheSelectionData(const int32 PaintLODIndex, const int32 UVChannel);
@@ -268,6 +257,10 @@ public:
 	void Refresh();
 	bool SelectionContainsPerLODColors() const { return bSelectionContainsPerLODColors; }
 	void ClearSelectionLODColors() { bSelectionContainsPerLODColors = false; }
+	void UpdatePaintSupportState();
+	bool GetSelectionSupportsVertexPaint() const { return bSelectionSupportsVertexPaint; }
+	bool GetSelectionSupportsTextureColorPaint() const { return bSelectionSupportsTextureColorPaint; }
+	bool GetSelectionSupportsTextureAssetPaint() const { return bSelectionSupportsTextureAssetPaint; }
 
 	FImage const& GetCopiedTexture() const;
 	void SetCopiedTexture(UTexture* InTexture);
@@ -278,9 +271,6 @@ public:
 	UPROPERTY(Transient)
 	TObjectPtr<UMeshComponent> LastPaintedComponent;
 
-protected:
-	bool bSelectionHasMaterialValidForTexturePaint;
-
 private:
 	void CleanUp();
 
@@ -300,6 +290,10 @@ private:
 
 	/** Contains copied texture data */
 	FImage CopiedTextureData;
+
+	bool bSelectionSupportsVertexPaint;
+	bool bSelectionSupportsTextureColorPaint;
+	bool bSelectionSupportsTextureAssetPaint;
 };
 
 template<typename T>
diff --git a/Engine/Plugins/MeshPainting/Source/MeshPaintingToolset/Public/MeshSelect.h b/Engine/Plugins/MeshPainting/Source/MeshPaintingToolset/Public/MeshSelect.h
--- a/Engine/Plugins/MeshPainting/Source/MeshPaintingToolset/Public/MeshSelect.h
+++ b/Engine/Plugins/MeshPainting/Source/MeshPaintingToolset/Public/MeshSelect.h
@@ -65,7 +65,7 @@ public:
  * The function ::OnClicked() implements the action that will occur when a click happens.
  * You must override this to implement any kind of useful behavior.
  */
-UCLASS()
+UCLASS(Abstract)
 class MESHPAINTINGTOOLSET_API UMeshClickTool : public USingleClickTool, public IMeshPaintSelectionInterface
 {
 	GENERATED_BODY()
@@ -78,14 +78,8 @@ public:
 	virtual void OnUpdateModifierState(int ModifierID, bool bIsOn) override;
 	virtual FInputRayHit IsHitByClick(const FInputDeviceRay& ClickPos) override;
 	virtual void OnClicked(const FInputDeviceRay& ClickPos) override;
-	virtual bool IsMeshAdapterSupported(TSharedPtr<IMeshPaintComponentAdapter> MeshAdapter) const override
-	{
-		return true;
-	}
-	virtual bool AllowsMultiselect() const override
-	{
-		return true;
-	}
+	virtual bool IsMeshAdapterSupported(TSharedPtr<IMeshPaintComponentAdapter> MeshAdapter) const override;
+	virtual bool AllowsMultiselect() const override { return false; }
 
 protected:
 	// flags used to identify modifier keys/buttons
@@ -102,8 +96,7 @@ class MESHPAINTINGTOOLSET_API UVertexAdapterClickTool : public UMeshClickTool
 	GENERATED_BODY()
 
 public:
-	// USingleClickTool overrides
-	virtual bool IsMeshAdapterSupported(TSharedPtr<IMeshPaintComponentAdapter> MeshAdapter) const override;
+	virtual bool AllowsMultiselect() const override { return true; }
 };
 
 UCLASS()
@@ -112,19 +105,13 @@ class MESHPAINTINGTOOLSET_API UTextureColorAdapterClickTool : public UMeshClickT
 	GENERATED_BODY()
 
 public:
-	// USingleClickTool overrides
-	virtual bool IsMeshAdapterSupported(TSharedPtr<IMeshPaintComponentAdapter> MeshAdapter) const override;
+	virtual bool AllowsMultiselect() const override { return true; }
 };
 
 UCLASS()
 class MESHPAINTINGTOOLSET_API UTextureAssetAdapterClickTool : public UMeshClickTool
 {
 	GENERATED_BODY()
-
-public:
-	// USingleClickTool overrides
-	virtual bool IsMeshAdapterSupported(TSharedPtr<IMeshPaintComponentAdapter> MeshAdapter) const override;
-	virtual bool AllowsMultiselect() const override { return false; }
 };
 
 #if UE_ENABLE_INCLUDE_ORDER_DEPRECATED_IN_5_2
diff --git a/Engine/Plugins/MeshPainting/Source/MeshPaintingToolset/Public/MeshVertexPaintingTool.h b/Engine/Plugins/MeshPainting/Source/MeshPaintingToolset/Public/MeshVertexPaintingTool.h
--- a/Engine/Plugins/MeshPainting/Source/MeshPaintingToolset/Public/MeshVertexPaintingTool.h
+++ b/Engine/Plugins/MeshPainting/Source/MeshPaintingToolset/Public/MeshVertexPaintingTool.h
@@ -48,7 +48,6 @@ enum class EMeshPaintTextureIndex : uint8
 	TextureFive
 };
 
-DECLARE_DELEGATE_RetVal(TArray<UMeshComponent*>, FGetSelectedMeshComponents);
 /**
  *
  */
