Make suggested tidy ups to Mesh Paint Mode Panel UI.
Move color channels mask to the top, just  below color channel visualization.

diff --git a/Engine/Plugins/MeshPainting/Source/MeshPaintEditorMode/Private/BrushSettingsCustomization.cpp b/Engine/Plugins/MeshPainting/Source/MeshPaintEditorMode/Private/BrushSettingsCustomization.cpp
--- a/Engine/Plugins/MeshPainting/Source/MeshPaintEditorMode/Private/BrushSettingsCustomization.cpp
+++ b/Engine/Plugins/MeshPainting/Source/MeshPaintEditorMode/Private/BrushSettingsCustomization.cpp
@@ -154,6 +154,8 @@ void FVertexColorPaintingSettingsCustomization::CustomizeDetails(IDetailLayoutBu
 	FVertexPaintingSettingsCustomization::CustomizeDetails(DetailLayout);
 
 	IDetailCategoryBuilder& ColorCategory = DetailLayout.EditCategory(TEXT("ColorPainting"));
+	ColorCategory.SetSortOrder(0);
+	IDetailCategoryBuilder& VertexCategory = DetailLayout.EditCategory(TEXT("VertexPainting"));
 
 	/** Creates a custom widget row containing all color channel flags */
 	TSharedRef<IPropertyHandle> RedChannel = DetailLayout.GetProperty("bWriteRed", UMeshVertexColorPaintingToolProperties::StaticClass());
@@ -200,7 +202,7 @@ void FVertexColorPaintingSettingsCustomization::CustomizeDetails(IDetailLayoutBu
 	LODPaintingIndex->MarkHiddenByCustomization();
 	TSharedPtr<SWidget> LODIndexWidget = LODPaintingIndex->CreatePropertyValueWidget();
 
-	ColorCategory.AddCustomRow(NSLOCTEXT("LODPainting", "LODPaintingLabel", "LOD Model Painting"))
+	VertexCategory.AddCustomRow(NSLOCTEXT("LODPainting", "LODPaintingLabel", "LOD Model Painting"))
 		.NameContent()
 		[
 			SNew(STextBlock)
@@ -256,7 +258,7 @@ void FVertexColorPaintingSettingsCustomization::CustomizeDetails(IDetailLayoutBu
 		]
 		];
 
-	ColorCategory.AddCustomRow(NSLOCTEXT("LODPainting", "LODPaintingLabel", "LOD Model Painting"))
+	VertexCategory.AddCustomRow(NSLOCTEXT("LODPainting", "LODPaintingLabel", "LOD Model Painting"))
 		.WholeRowContent()
 		[
 			SNew(SWarningOrErrorBox)
@@ -298,6 +300,7 @@ void FVertexWeightPaintingSettingsCustomization::CustomizeDetails(IDetailLayoutB
 	FVertexPaintingSettingsCustomization::CustomizeDetails(DetailLayout);
 
 	IDetailCategoryBuilder& WeightCategory = DetailLayout.EditCategory(TEXT("WeightPainting"));
+	WeightCategory.SetSortOrder(0);
 
 	TSharedRef<IPropertyHandle> WeightTypeProperty = DetailLayout.GetProperty(GET_MEMBER_NAME_CHECKED(UMeshVertexWeightPaintingToolProperties, TextureWeightType));
 	TSharedRef<IPropertyHandle> PaintWeightProperty = DetailLayout.GetProperty(GET_MEMBER_NAME_CHECKED(UMeshVertexWeightPaintingToolProperties, PaintTextureWeightIndex));
@@ -350,9 +353,10 @@ void FTexturePaintingSettingsCustomization::CustomizeDetails(IDetailLayoutBuilde
 {
 	FMeshPaintingSettingsCustomization::CustomizeDetails(DetailLayout);
 
+	IDetailCategoryBuilder& ColorCategory = DetailLayout.EditCategory(TEXT("ColorPainting"));
+	ColorCategory.SetSortOrder(0);
 	IDetailCategoryBuilder& BrushCategory = DetailLayout.EditCategory(TEXT("Brush"), FText::GetEmpty(), ECategoryPriority::Important);
 	IDetailCategoryBuilder& TextureCategory = DetailLayout.EditCategory(TEXT("TexturePainting"));
-	IDetailCategoryBuilder& ColorCategory = DetailLayout.EditCategory(TEXT("ColorPainting"));
 
 	TSharedRef<IPropertyHandle> RedChannel = DetailLayout.GetProperty("bWriteRed", UMeshTexturePaintingToolProperties::StaticClass());
 	RedChannel->MarkHiddenByCustomization();
diff --git a/Engine/Plugins/MeshPainting/Source/MeshPaintingToolset/Public/MeshVertexPaintingTool.h b/Engine/Plugins/MeshPainting/Source/MeshPaintingToolset/Public/MeshVertexPaintingTool.h
--- a/Engine/Plugins/MeshPainting/Source/MeshPaintingToolset/Public/MeshVertexPaintingTool.h
+++ b/Engine/Plugins/MeshPainting/Source/MeshPaintingToolset/Public/MeshVertexPaintingTool.h
@@ -110,11 +110,11 @@ public:
 	bool bWriteAlpha = false;
 
 	/** When unchecked the painting on the base LOD will be propagate automatically to all other LODs when exiting the mode or changing the selection */
-	UPROPERTY(EditAnywhere, Category = Painting, meta = (TransientToolProperty))
+	UPROPERTY(EditAnywhere, Category = VertexPainting, meta = (TransientToolProperty))
 	bool bPaintOnSpecificLOD = false;
 
 	/** LOD Index to which should specifically be painted */
-	UPROPERTY(EditAnywhere, Category = Painting, meta = (UIMin = "0", ClampMin = "0", EditCondition = "bPaintOnSpecificLOD", TransientToolProperty))
+	UPROPERTY(EditAnywhere, Category = VertexPainting, meta = (UIMin = "0", ClampMin = "0", EditCondition = "bPaintOnSpecificLOD", TransientToolProperty))
 	int32 LODIndex = 0;
 };
 
