MeshPaintTool: Hide unused brush colors for vertex weight painting

diff --git a/Engine/Plugins/MeshPainting/Source/MeshPaintEditorMode/Private/BrushSettingsCustomization.cpp b/Engine/Plugins/MeshPainting/Source/MeshPaintEditorMode/Private/BrushSettingsCustomization.cpp
--- a/Engine/Plugins/MeshPainting/Source/MeshPaintEditorMode/Private/BrushSettingsCustomization.cpp
+++ b/Engine/Plugins/MeshPainting/Source/MeshPaintEditorMode/Private/BrushSettingsCustomization.cpp
@@ -131,8 +131,6 @@ TSharedRef<IDetailCustomization> FVertexPaintingSettingsCustomization::MakeInsta
 
 void FVertexPaintingSettingsCustomization::CustomizeDetails(IDetailLayoutBuilder& DetailLayout)
 {
-	FMeshPaintingSettingsCustomization::CustomizeDetails(DetailLayout);
-
 	IDetailCategoryBuilder& VertexCategory = DetailLayout.EditCategory(TEXT("VertexPainting"));
 
 	VertexCategory.AddCustomRow(NSLOCTEXT("VertexPaintSettings", "InstanceColorSize", "Instance Color Size"))
@@ -151,6 +149,7 @@ TSharedRef<IDetailCustomization> FVertexColorPaintingSettingsCustomization::Make
 
 void FVertexColorPaintingSettingsCustomization::CustomizeDetails(IDetailLayoutBuilder& DetailLayout)
 {
+	FMeshPaintingSettingsCustomization::CustomizeDetails(DetailLayout);
 	FVertexPaintingSettingsCustomization::CustomizeDetails(DetailLayout);
 
 	IDetailCategoryBuilder& ColorCategory = DetailLayout.EditCategory(TEXT("ColorPainting"));
@@ -299,6 +298,13 @@ void FVertexWeightPaintingSettingsCustomization::CustomizeDetails(IDetailLayoutB
 {
 	FVertexPaintingSettingsCustomization::CustomizeDetails(DetailLayout);
 
+	// Brush colors not needed for weight painting.
+	IDetailCategoryBuilder& BrushCategory = DetailLayout.EditCategory(TEXT("Brush"), FText::GetEmpty(), ECategoryPriority::Important);
+	TSharedRef<IPropertyHandle> PaintColor = DetailLayout.GetProperty("PaintColor", UMeshPaintingToolProperties::StaticClass());
+	PaintColor->MarkHiddenByCustomization();
+	TSharedRef<IPropertyHandle> EraseColor = DetailLayout.GetProperty("EraseColor", UMeshPaintingToolProperties::StaticClass());
+	EraseColor->MarkHiddenByCustomization();
+
 	IDetailCategoryBuilder& WeightCategory = DetailLayout.EditCategory(TEXT("WeightPainting"));
 	WeightCategory.SetSortOrder(0);
 
diff --git a/Engine/Plugins/MeshPainting/Source/MeshPaintEditorMode/Private/MeshPaintModeCommands.cpp b/Engine/Plugins/MeshPainting/Source/MeshPaintEditorMode/Private/MeshPaintModeCommands.cpp
--- a/Engine/Plugins/MeshPainting/Source/MeshPaintEditorMode/Private/MeshPaintModeCommands.cpp
+++ b/Engine/Plugins/MeshPainting/Source/MeshPaintEditorMode/Private/MeshPaintModeCommands.cpp
@@ -66,7 +66,7 @@ void FMeshPaintEditorModeCommands::RegisterCommands()
 	Commands.Add(UMeshPaintMode::MeshPaintMode_VertexColor, VertexColorCommands);
 
 	TArray<TSharedPtr<FUICommandInfo>> VertexWeightCommands = {
-		SelectVertex, PaintVertexWeight, SwapColor, FillVertex, Propagate, ApplyVertex, SaveVertex, RemoveVertex, Copy, Paste, Import, Fix
+		SelectVertex, PaintVertexWeight, FillVertex, Propagate, ApplyVertex, SaveVertex, RemoveVertex, Copy, Paste, Import, Fix
 	};
 	Commands.Add(UMeshPaintMode::MeshPaintMode_VertexWeights, VertexWeightCommands);
 
