Added option to update vertex colors while painting texture colors.
This is enabled with the Propagate To Vertex Color setting in the paint tool.
This provides a way to automatically keep vertex painting fallback in sync with texture color painting.

diff --git a/Engine/Plugins/MeshPainting/Source/MeshPaintEditorMode/Private/MeshPaintMode.cpp b/Engine/Plugins/MeshPainting/Source/MeshPaintEditorMode/Private/MeshPaintMode.cpp
--- a/Engine/Plugins/MeshPainting/Source/MeshPaintEditorMode/Private/MeshPaintMode.cpp
+++ b/Engine/Plugins/MeshPainting/Source/MeshPaintEditorMode/Private/MeshPaintMode.cpp
@@ -890,10 +890,23 @@ bool UMeshPaintMode::CanCycleTextures() const
 
 void UMeshPaintMode::ApplyTextureColorsToAsset()
 {
+	FScopedTransaction Transaction(LOCTEXT("MeshPaintMode_TexturePaint_Apply", "Apply Texture Paint"));
+
+	TArray<UMeshComponent*> VertexColorPropagateComponents;
+	if (UMeshTextureColorPaintingTool* TextureColorPaintingTool = Cast<UMeshTextureColorPaintingTool>(GetToolManager()->GetActiveTool(EToolSide::Left)))
+	{
+		TextureColorPaintingTool->GetPaintedComponentsForPropagateVertexColor(VertexColorPropagateComponents);
+	}
+
 	if (UMeshTexturePaintingTool* TexturePaintingTool = Cast<UMeshTexturePaintingTool>(GetToolManager()->GetActiveTool(EToolSide::Left)))
 	{
 		TexturePaintingTool->ApplyAllPaintedTextures();
 	}
+
+	for (UMeshComponent* Component : VertexColorPropagateComponents)
+	{
+		GEditor->GetEditorSubsystem<UMeshPaintModeSubsystem>()->ImportVertexColorsFromMeshPaintTexture(Component);
+	}
 }
 
 bool UMeshPaintMode::CanApplyTextureColorsToAsset() const
diff --git a/Engine/Plugins/MeshPainting/Source/MeshPaintEditorMode/Private/MeshPaintModeHelpers.cpp b/Engine/Plugins/MeshPainting/Source/MeshPaintEditorMode/Private/MeshPaintModeHelpers.cpp
--- a/Engine/Plugins/MeshPainting/Source/MeshPaintEditorMode/Private/MeshPaintModeHelpers.cpp
+++ b/Engine/Plugins/MeshPainting/Source/MeshPaintEditorMode/Private/MeshPaintModeHelpers.cpp
@@ -281,6 +281,12 @@ void UMeshPaintModeSubsystem::ImportVertexColorsFromMeshPaintTexture(UMeshCompon
 {
 	if (UTexture2D* Texture = Cast<UTexture2D>(MeshComponent->GetMeshPaintTexture()))
 	{
+#if WITH_EDITOR
+		// We may need to wait for the texture to compile before importing.
+		// This is most likely to happen when we are immediately propagating texture color painting to vertex colors.
+		Texture->BlockOnAnyAsyncBuild();
+#endif
+
 		UImportVertexColorOptions* Options = NewObject<UImportVertexColorOptions>();
  		Options->UVIndex = MeshComponent->GetMeshPaintTextureCoordinateIndex();
 
diff --git a/Engine/Plugins/MeshPainting/Source/MeshPaintingToolset/Private/MeshTexturePaintingTool.cpp b/Engine/Plugins/MeshPainting/Source/MeshPaintingToolset/Private/MeshTexturePaintingTool.cpp
--- a/Engine/Plugins/MeshPainting/Source/MeshPaintingToolset/Private/MeshTexturePaintingTool.cpp
+++ b/Engine/Plugins/MeshPainting/Source/MeshPaintingToolset/Private/MeshTexturePaintingTool.cpp
@@ -1588,10 +1588,7 @@ void UMeshTexturePaintingTool::ApplyAllPaintedTextures()
 	{
 		check(PaintingTexture2D == nullptr);
 
-		FScopedTransaction Transaction(LOCTEXT("MeshPaintMode_TexturePaint_Apply", "Apply Texture Paint"));
-		
 		Modify();
-	//	GWarn->BeginSlowTask(LOCTEXT("BeginMeshPaintMode_TexturePaint_ApplyTask", "Applying Texture Paint Changes"), true);
 
 		int32 CurStep = 1;
 		int32 TotalSteps = GetNumberOfPendingPaintChanges();
@@ -1603,8 +1600,6 @@ void UMeshTexturePaintingTool::ApplyAllPaintedTextures()
 			// Apply the texture
 			if (TextureData->bIsPaintingTexture2DModified == true)
 			{
-			//	GWarn->StatusUpdate(CurStep++, TotalSteps, FText::Format(LOCTEXT("MeshPaintMode_TexturePaint_ApplyStatus", "Applying Texture Paint Changes: {0}"), FText::FromName(TextureData->PaintingTexture2D->GetFName())));
-
 				const int32 TexWidth = TextureData->PaintRenderTargetTexture->SizeX;
 				const int32 TexHeight = TextureData->PaintRenderTargetTexture->SizeY;
 				TArray< FColor > TexturePixels;
@@ -1645,8 +1640,6 @@ void UMeshTexturePaintingTool::ApplyAllPaintedTextures()
 		}
 
 		ClearAllTextureOverrides();
-
-//		GWarn->EndSlowTask();
 	}
 }
 
@@ -1683,7 +1676,6 @@ int32 UMeshTexturePaintingTool::GetNumberOfPendingPaintChanges() const
 	{
 		const FPaintTexture2DData* TextureData = &It.Value();
 
-		// Apply the texture
 		if (TextureData->bIsPaintingTexture2DModified == true)
 		{
 			Result++;
@@ -1758,6 +1750,30 @@ void UMeshTextureColorPaintingTool::CacheTexturePaintData()
 	}
 }
 
+void UMeshTextureColorPaintingTool::GetPaintedComponentsForPropagateVertexColor(TArray<UMeshComponent*>& OutComponents) const
+{
+	if (!ColorProperties->bPropagateToVertexColor)
+	{
+		return;
+	}
+	
+	for (decltype(PaintTargetData)::TConstIterator It(PaintTargetData); It; ++It)
+	{
+		const FPaintTexture2DData* TextureData = &It.Value();
+
+		if (TextureData->bIsPaintingTexture2DModified == true)
+		{
+			for (UMeshComponent* Component : TextureData->PaintedComponents)
+			{
+				if (Component->GetMeshPaintTexture() == TextureData->PaintingTexture2D)
+				{
+					OutComponents.Add(Component);
+				}
+			}
+		}
+	}
+}
+
 
 UMeshTextureAssetPaintingTool::UMeshTextureAssetPaintingTool()
 {
diff --git a/Engine/Plugins/MeshPainting/Source/MeshPaintingToolset/Public/MeshTexturePaintingTool.h b/Engine/Plugins/MeshPainting/Source/MeshPaintingToolset/Public/MeshTexturePaintingTool.h
--- a/Engine/Plugins/MeshPainting/Source/MeshPaintingToolset/Public/MeshTexturePaintingTool.h
+++ b/Engine/Plugins/MeshPainting/Source/MeshPaintingToolset/Public/MeshTexturePaintingTool.h
@@ -104,6 +104,11 @@ UCLASS()
 class MESHPAINTINGTOOLSET_API UMeshTextureColorPaintingToolProperties : public UMeshTexturePaintingToolProperties
 {
 	GENERATED_BODY()
+
+public:
+	/** Whether to copy texture color painting to vertex colors on Apply. */
+	UPROPERTY(EditAnywhere, Category = ColorPainting)
+	bool bPropagateToVertexColor = false;
 };
 
 /**
@@ -286,6 +291,9 @@ public:
 	virtual void CacheTexturePaintData() override;
 	// End UMeshTexturePaintingTool Interface.
 
+	/** Get mesh components that have mesh paint textures pending save and that should write vertex colors when they do save. */
+	void GetPaintedComponentsForPropagateVertexColor(TArray<UMeshComponent*>& OutComponents) const;
+
 protected:
 	UPROPERTY(Transient)
 	TObjectPtr<UMeshTextureColorPaintingToolProperties> ColorProperties;
diff --git a/Engine/Source/Runtime/Engine/Classes/Components/StaticMeshComponent.h b/Engine/Source/Runtime/Engine/Classes/Components/StaticMeshComponent.h
--- a/Engine/Source/Runtime/Engine/Classes/Components/StaticMeshComponent.h
+++ b/Engine/Source/Runtime/Engine/Classes/Components/StaticMeshComponent.h
@@ -313,13 +313,13 @@ public:
 	uint8 bReverseCulling : 1;
 
 	/** 
-	 * If false, vertex color mesh painting is disabled on this instance. 
+	 * If false, vertex color mesh painting is disabled on this component. 
 	 * This may be set to false by blueprint functions that override vertex colors in construction script.
 	 */
 	UPROPERTY(EditAnywhere, BlueprintReadWrite, Category = "Mesh Painting")
 	uint8 bEnableVertexColorMeshPainting : 1;
 
-	/** If false, texture color mesh painting is disabled on this instance. */
+	/** If false, texture color mesh painting is disabled on this component. */
 	UPROPERTY(EditAnywhere, BlueprintReadWrite, Category = "Mesh Painting")
 	uint8 bEnableTextureColorMeshPainting : 1;
 
