Add mesh paint mode option to fix texture colors.
This detects if any selected mesh paint textures are not at the currently requested size and resizes them.
It's similar in spirit to the vertex painting "fix" operation which detects if the mesh has changed and tries to reapply painting.

diff --git a/Engine/Plugins/MeshPainting/Source/MeshPaintEditorMode/Private/MeshPaintMode.cpp b/Engine/Plugins/MeshPainting/Source/MeshPaintEditorMode/Private/MeshPaintMode.cpp
--- a/Engine/Plugins/MeshPainting/Source/MeshPaintEditorMode/Private/MeshPaintMode.cpp
+++ b/Engine/Plugins/MeshPainting/Source/MeshPaintEditorMode/Private/MeshPaintMode.cpp
@@ -338,10 +338,14 @@ void UMeshPaintMode::BindCommands()
 		FExecuteAction::CreateUObject(this, &UMeshPaintMode::ImportMeshPaintTextureFromVertexColors),
 		FCanExecuteAction::CreateUObject(this, &UMeshPaintMode::CanImportMeshPaintTextureFromVertexColors)));
 
-	CommandList->MapAction(Commands.Fix, FUIAction(
+	CommandList->MapAction(Commands.FixVertex, FUIAction(
 		FExecuteAction::CreateUObject(this, &UMeshPaintMode::FixVertexColors),
 		FCanExecuteAction::CreateUObject(this, &UMeshPaintMode::CanFixVertexColors)));
 
+	CommandList->MapAction(Commands.FixTexture, FUIAction(
+		FExecuteAction::CreateUObject(this, &UMeshPaintMode::FixTextureColors),
+		FCanExecuteAction::CreateUObject(this, &UMeshPaintMode::CanFixTextureColors)));
+
 	CommandList->MapAction(Commands.PreviousLOD, FUIAction(
 		FExecuteAction::CreateUObject(this, &UMeshPaintMode::CycleMeshLODs, -1),
 		FCanExecuteAction::CreateUObject(this, &UMeshPaintMode::CanCycleMeshLODs)));
@@ -1275,6 +1279,22 @@ bool UMeshPaintMode::CanImportMeshPaintTextureFromVertexColors() const
 	return false;
 }
 
+void UMeshPaintMode::FixTextureColors()
+{
+	FScopedTransaction Transaction(LOCTEXT("LevelMeshPainter_TransactionFixTextureColors", "Fixing Per-Instance Texture Colors"));
+	GEditor->GetEditorSubsystem<UMeshPaintModeSubsystem>()->FixTextureColors(GetSelectedComponents<UMeshComponent>());
+}
+
+bool UMeshPaintMode::CanFixTextureColors() const
+{
+	if (!IsInSelectTool())
+	{
+		return false;
+	}
+
+	return GEditor->GetEditorSubsystem<UMeshPaintModeSubsystem>()->CanFixTextureColors(GetSelectedComponents<UMeshComponent>());
+}
+
 bool UMeshPaintMode::CanCycleMeshLODs() const
 {
 	FName PaletteName = Toolkit->GetCurrentPalette();
diff --git a/Engine/Plugins/MeshPainting/Source/MeshPaintEditorMode/Private/MeshPaintMode.h b/Engine/Plugins/MeshPainting/Source/MeshPaintEditorMode/Private/MeshPaintMode.h
--- a/Engine/Plugins/MeshPainting/Source/MeshPaintEditorMode/Private/MeshPaintMode.h
+++ b/Engine/Plugins/MeshPainting/Source/MeshPaintEditorMode/Private/MeshPaintMode.h
@@ -134,6 +134,8 @@ protected:
 	bool CanImportMeshPaintTextureFromVertexColors() const;
 	void FixVertexColors();
 	bool CanFixVertexColors() const;
+	void FixTextureColors();
+	bool CanFixTextureColors() const;
 	void CycleMeshLODs(int32 Direction);
 	bool CanCycleMeshLODs() const;
 	void CycleTextures(int32 Direction);
diff --git a/Engine/Plugins/MeshPainting/Source/MeshPaintEditorMode/Private/MeshPaintModeCommands.cpp b/Engine/Plugins/MeshPainting/Source/MeshPaintEditorMode/Private/MeshPaintModeCommands.cpp
--- a/Engine/Plugins/MeshPainting/Source/MeshPaintEditorMode/Private/MeshPaintModeCommands.cpp
+++ b/Engine/Plugins/MeshPainting/Source/MeshPaintEditorMode/Private/MeshPaintModeCommands.cpp
@@ -54,24 +54,25 @@ void FMeshPaintEditorModeCommands::RegisterCommands()
 	UI_COMMAND(Import, "Import", "Imports vertex colors from a TGA texture file to the selected meshes", EUserInterfaceActionType::Button, FInputChord());
 	UI_COMMAND(GetVertexColors, "Vertex", "Imports texture colors from vertex colors on the selected mesh components", EUserInterfaceActionType::Button, FInputChord());
 	UI_COMMAND(GetTextureColors, "Texture", "Imports vertex colors from texture colors on the selected mesh components", EUserInterfaceActionType::Button, FInputChord());
-	UI_COMMAND(Fix, "Fix", "Applies any required color data fixes to the selected mesh components", EUserInterfaceActionType::Button, FInputChord());
+	UI_COMMAND(FixVertex, "Fix", "Applies any required color data fixes to the selected mesh components", EUserInterfaceActionType::Button, FInputChord());
+	UI_COMMAND(FixTexture, "Fix", "Applies any pending resolution change of texture color painting", EUserInterfaceActionType::Button, FInputChord());
 	UI_COMMAND(PreviousLOD, "Previous LOD", "Cycles to the previous possible mesh LOD to paint on", EUserInterfaceActionType::Button, FInputChord(EKeys::B));
 	UI_COMMAND(NextLOD, "Next LOD", "Cycles to the next possible mesh LOD to paint on", EUserInterfaceActionType::Button, FInputChord(EKeys::N));
 	UI_COMMAND(PreviousTexture, "Previous Texture", "Cycle To previous texture", EUserInterfaceActionType::Button, FInputChord(EKeys::Comma));
 	UI_COMMAND(NextTexture, "Next Texture", "Cycle To next texture", EUserInterfaceActionType::Button, FInputChord(EKeys::Period));
 
 	TArray<TSharedPtr<FUICommandInfo>> VertexColorCommands = {
-		SelectVertex, PaintVertexColor, SwapColor, FillVertex, Propagate, ApplyVertex, SaveVertex, RemoveVertex, Copy, Paste, Import, GetTextureColors, Fix
+		SelectVertex, PaintVertexColor, SwapColor, FillVertex, Propagate, ApplyVertex, SaveVertex, RemoveVertex, Copy, Paste, Import, GetTextureColors, FixVertex
 	};
 	Commands.Add(UMeshPaintMode::MeshPaintMode_VertexColor, VertexColorCommands);
 
 	TArray<TSharedPtr<FUICommandInfo>> VertexWeightCommands = {
-		SelectVertex, PaintVertexWeight, FillVertex, Propagate, ApplyVertex, SaveVertex, RemoveVertex, Copy, Paste, Import, Fix
+		SelectVertex, PaintVertexWeight, FillVertex, Propagate, ApplyVertex, SaveVertex, RemoveVertex, Copy, Paste, Import, FixVertex
 	};
 	Commands.Add(UMeshPaintMode::MeshPaintMode_VertexWeights, VertexWeightCommands);
 
 	TArray<TSharedPtr<FUICommandInfo>> TextureColorCommands = {
-		SelectTextureColor, PaintTextureColor, SwapColor, FillTexture, ApplyTexture, SaveTexture, Add, RemoveTexture, Copy, Paste, GetVertexColors
+		SelectTextureColor, PaintTextureColor, SwapColor, FillTexture, ApplyTexture, SaveTexture, Add, RemoveTexture, Copy, Paste, GetVertexColors, FixTexture
 	};
 	Commands.Add(UMeshPaintMode::MeshPaintMode_TextureColor, TextureColorCommands);
 
diff --git a/Engine/Plugins/MeshPainting/Source/MeshPaintEditorMode/Private/MeshPaintModeHelpers.cpp b/Engine/Plugins/MeshPainting/Source/MeshPaintEditorMode/Private/MeshPaintModeHelpers.cpp
--- a/Engine/Plugins/MeshPainting/Source/MeshPaintEditorMode/Private/MeshPaintModeHelpers.cpp
+++ b/Engine/Plugins/MeshPainting/Source/MeshPaintEditorMode/Private/MeshPaintModeHelpers.cpp
@@ -907,6 +907,55 @@ void UMeshPaintModeSubsystem::RemovePerLODColors(const TArray<UMeshComponent*>&
 	}
 }
 
+bool UMeshPaintModeSubsystem::CanFixTextureColors(const TArray<UMeshComponent*>& Components)
+{
+	for (UMeshComponent* Component : Components)
+	{
+		if (UStaticMeshComponent* StaticMeshComponent = Cast<UStaticMeshComponent>(Component))
+		{
+			if (UTexture* Texture = StaticMeshComponent->GetMeshPaintTexture())
+			{
+				if (StaticMeshComponent->GetMeshPaintTextureResolution() != Texture->Source.GetSizeX())
+				{
+					return true;
+				}
+			}
+		}
+	}
+	return false;
+}
+
+void UMeshPaintModeSubsystem::FixTextureColors(const TArray<UMeshComponent*>& Components)
+{
+	for (UMeshComponent* Component : Components)
+	{
+		if (UStaticMeshComponent* StaticMeshComponent = Cast<UStaticMeshComponent>(Component))
+		{
+			if (UTexture* Texture = StaticMeshComponent->GetMeshPaintTexture())
+			{
+				int32 TextureResolution = StaticMeshComponent->GetMeshPaintTextureResolution();
+				if (TextureResolution != Texture->Source.GetSizeX())
+				{
+					FImage Image;
+					if (Texture->Source.GetMipImage(Image, 0))
+					{
+						FImage ResizedImage(TextureResolution, TextureResolution, Image.NumSlices, Image.Format, Image.GammaSpace);
+						FImageCore::ResizeImage(Image, ResizedImage);
+
+						UMeshPaintVirtualTexture* NewTexture = NewObject<UMeshPaintVirtualTexture>(StaticMeshComponent->GetOutermost());
+						NewTexture->Source.Init(ResizedImage);
+						NewTexture->OwningComponent = MakeWeakObjectPtr(StaticMeshComponent);
+						NewTexture->UpdateResource();
+
+						StaticMeshComponent->Modify();
+						StaticMeshComponent->SetMeshPaintTexture(NewTexture);
+					}
+				}
+			}
+		}
+	}
+}
+
 void UMeshPaintModeSubsystem::SwapColors()
 {
 	if (UMeshPaintingToolProperties* Settings = UMeshPaintMode::GetToolProperties())
diff --git a/Engine/Plugins/MeshPainting/Source/MeshPaintEditorMode/Public/MeshPaintModeCommands.h b/Engine/Plugins/MeshPainting/Source/MeshPaintEditorMode/Public/MeshPaintModeCommands.h
--- a/Engine/Plugins/MeshPainting/Source/MeshPaintEditorMode/Public/MeshPaintModeCommands.h
+++ b/Engine/Plugins/MeshPainting/Source/MeshPaintEditorMode/Public/MeshPaintModeCommands.h
@@ -55,7 +55,8 @@ public:
 	TSharedPtr<FUICommandInfo> Import;
 	TSharedPtr<FUICommandInfo> GetTextureColors;
 	TSharedPtr<FUICommandInfo> GetVertexColors;
-	TSharedPtr<FUICommandInfo> Fix;
+	TSharedPtr<FUICommandInfo> FixVertex;
+	TSharedPtr<FUICommandInfo> FixTexture;
 
 	TSharedPtr<FUICommandInfo> PreviousLOD;
 	TSharedPtr<FUICommandInfo> NextLOD;
diff --git a/Engine/Plugins/MeshPainting/Source/MeshPaintEditorMode/Public/MeshPaintModeHelpers.h b/Engine/Plugins/MeshPainting/Source/MeshPaintEditorMode/Public/MeshPaintModeHelpers.h
--- a/Engine/Plugins/MeshPainting/Source/MeshPaintEditorMode/Public/MeshPaintModeHelpers.h
+++ b/Engine/Plugins/MeshPainting/Source/MeshPaintEditorMode/Public/MeshPaintModeHelpers.h
@@ -82,6 +82,9 @@ public:
 	void PasteVertexColors(const TArray<UStaticMeshComponent*>& StaticMeshComponents, TArray<FPerComponentVertexColorData>& CopiedColorsByComponent);
 	bool CanPasteInstanceVertexColors(const TArray<UStaticMeshComponent*>& StaticMeshComponents, const TArray<FPerComponentVertexColorData>& CopiedColorsByComponent);
 	void RemovePerLODColors(const TArray<UMeshComponent*>& PaintableComponents);
+	
+	bool CanFixTextureColors(const TArray<UMeshComponent*>& Components);
+	void FixTextureColors(const TArray<UMeshComponent*>& Components);
 
 	void SwapColors();
 };
diff --git a/Engine/Source/Editor/EditorStyle/Private/StarshipStyle.cpp b/Engine/Source/Editor/EditorStyle/Private/StarshipStyle.cpp
--- a/Engine/Source/Editor/EditorStyle/Private/StarshipStyle.cpp
+++ b/Engine/Source/Editor/EditorStyle/Private/StarshipStyle.cpp
@@ -4974,8 +4974,10 @@ void FStarshipEditorStyle::FStyle::SetupLevelEditorStyle()
 		Set("MeshPaint.GetTextureColors.Small", new IMAGE_BRUSH("/Icons/GeneralTools/Import_40x", Icon20x20));
 		Set("MeshPaint.GetVertexColors", new IMAGE_BRUSH("/Icons/GeneralTools/Import_40x", Icon20x20));
 		Set("MeshPaint.GetVertexColors.Small", new IMAGE_BRUSH("/Icons/GeneralTools/Import_40x", Icon20x20));
-		Set("MeshPaint.Fix", new IMAGE_BRUSH("/Icons/GeneralTools/Fix_40x", Icon20x20));
-		Set("MeshPaint.Fix.Small", new IMAGE_BRUSH("/Icons/GeneralTools/Fix_40x", Icon20x20));
+		Set("MeshPaint.FixVertex", new IMAGE_BRUSH("/Icons/GeneralTools/Fix_40x", Icon20x20));
+		Set("MeshPaint.FixVertex.Small", new IMAGE_BRUSH("/Icons/GeneralTools/Fix_40x", Icon20x20));
+		Set("MeshPaint.FixTexture", new IMAGE_BRUSH("/Icons/GeneralTools/Fix_40x", Icon20x20));
+		Set("MeshPaint.FixTexture.Small", new IMAGE_BRUSH("/Icons/GeneralTools/Fix_40x", Icon20x20));
 		Set("MeshPaint.PreviousLOD", new IMAGE_BRUSH(TEXT("Icons/GeneralTools/Previous_40x"), Icon20x20));
 		Set("MeshPaint.PreviousLOD.Small", new IMAGE_BRUSH(TEXT("Icons/GeneralTools/Previous_40x"), Icon20x20));
 		Set("MeshPaint.NextLOD", new IMAGE_BRUSH(TEXT("Icons/GeneralTools/Next_40x"), Icon20x20));
diff --git a/Engine/Source/Runtime/Engine/Classes/Engine/RendererSettings.h b/Engine/Source/Runtime/Engine/Classes/Engine/RendererSettings.h
--- a/Engine/Source/Runtime/Engine/Classes/Engine/RendererSettings.h
+++ b/Engine/Source/Runtime/Engine/Classes/Engine/RendererSettings.h
@@ -464,7 +464,7 @@ class URendererSettings : public UDeveloperSettings
 		EditCondition = "bVirtualTextures",
 		ConsoleVariable = "r.MeshPaintVirtualTexture.DefaultTexelsPerVertex", DisplayName = "Mesh paint texels per vertex",
 		ToolTip = "Default ratio of texels to vertices when creating a mesh paint texture for a mesh."))
-	uint32 MeshPaintDefaultTexelsPerVertex;
+	int32 MeshPaintDefaultTexelsPerVertex;
 
 	UPROPERTY(config, EditAnywhere, Category = "Runtime Virtual Textures|Enable Material Types", meta = (
 		EditCondition = "bVirtualTextures",
