Add mesh paint texture resource size to the mesh paint tool.
Added a new section to put this in along with the vertex paint data size.
We want to show this info to help with tracking data usage.
But it should also be noted that texture resource size is the uncompressed size on disk, which isn't the same as compressed size, and isn't the same as memory (particularly as these are virtual textures).

diff --git a/Engine/Plugins/MeshPainting/Source/MeshPaintEditorMode/Private/BrushSettingsCustomization.cpp b/Engine/Plugins/MeshPainting/Source/MeshPaintEditorMode/Private/BrushSettingsCustomization.cpp
--- a/Engine/Plugins/MeshPainting/Source/MeshPaintEditorMode/Private/BrushSettingsCustomization.cpp
+++ b/Engine/Plugins/MeshPainting/Source/MeshPaintEditorMode/Private/BrushSettingsCustomization.cpp
@@ -38,6 +38,53 @@ TSharedRef<SHorizontalBox> CreateColorChannelWidget(TSharedRef<IPropertyHandle>
 }
 
 
+TSharedRef<IDetailCustomization> FMeshPaintModeSettingsCustomization::MakeInstance()
+{
+	return MakeShareable(new FMeshPaintModeSettingsCustomization);
+}
+
+void FMeshPaintModeSettingsCustomization::CustomizeDetails(IDetailLayoutBuilder& DetailLayout)
+{
+	IDetailCategoryBuilder& ResourceCategory = DetailLayout.EditCategory(TEXT("ResourceUsage"));
+
+	ResourceCategory.AddCustomRow(NSLOCTEXT("VertexPaintSettings", "VertexColorSizeRow", "Vertex Color Size"))
+		.NameContent()
+		[
+			SNew(STextBlock)
+				.Font(IDetailLayoutBuilder::GetDetailFont())
+				.Text(NSLOCTEXT("VertexPaintSettings", "VertexColorSize", "Instance vertex color size"))
+		]
+		.ValueContent()
+		[
+			SNew(STextBlock)
+				.Font(IDetailLayoutBuilder::GetDetailFont())
+				.Text_Lambda([]() -> FText
+				{
+					const uint32 SizeInKB = FMath::DivideAndRoundNearest(UMeshPaintMode::GetMeshPaintMode()->GetVertexDataSizeInBytes(), 1024u);
+					return FText::Format(NSLOCTEXT("VertexPaintSettings", "VertexColorSizeValue", "{0} KB"), FText::AsNumber(SizeInKB));
+				})
+		];
+
+	ResourceCategory.AddCustomRow(NSLOCTEXT("VertexPaintSettings", "TextureColorSizeRow", "Texture Resource Size"))
+		.NameContent()
+		[
+			SNew(STextBlock)
+				.Font(IDetailLayoutBuilder::GetDetailFont())
+				.Text(NSLOCTEXT("VertexPaintSettings", "TextureColorSize", "Mesh paint texture resource size"))
+		]
+		.ValueContent()
+		[
+			SNew(STextBlock)
+				.Font(IDetailLayoutBuilder::GetDetailFont())
+				.Text_Lambda([]() -> FText 
+				{
+					const uint32 SizeInKB = FMath::DivideAndRoundNearest(UMeshPaintMode::GetMeshPaintMode()->GetMeshPaintTextureResourceSizeInBytes(), 1024u);
+					return FText::Format(NSLOCTEXT("VertexPaintSettings", "TextureColorSizeValue", "{0} KB"), FText::AsNumber(SizeInKB)); 
+				})
+		];
+}
+
+
 TSharedRef<IDetailCustomization> FMeshPaintingSettingsCustomization::MakeInstance()
 {
 	return MakeShareable(new FMeshPaintingSettingsCustomization);
@@ -131,14 +178,6 @@ TSharedRef<IDetailCustomization> FVertexPaintingSettingsCustomization::MakeInsta
 
 void FVertexPaintingSettingsCustomization::CustomizeDetails(IDetailLayoutBuilder& DetailLayout)
 {
-	IDetailCategoryBuilder& VertexCategory = DetailLayout.EditCategory(TEXT("VertexPainting"));
-
-	VertexCategory.AddCustomRow(NSLOCTEXT("VertexPaintSettings", "InstanceColorSize", "Instance Color Size"))
-		.WholeRowContent()
-		[
-			SNew(STextBlock)
-			.Text_Lambda([]() -> FText { return FText::Format(FTextFormat::FromString(TEXT("Instance Color Size: {0} KB")), UMeshPaintMode::GetMeshPaintMode()->GetCachedVertexDataSize() / 1024.f); })
-		];
 }
 
 
diff --git a/Engine/Plugins/MeshPainting/Source/MeshPaintEditorMode/Private/BrushSettingsCustomization.h b/Engine/Plugins/MeshPainting/Source/MeshPaintEditorMode/Private/BrushSettingsCustomization.h
--- a/Engine/Plugins/MeshPainting/Source/MeshPaintEditorMode/Private/BrushSettingsCustomization.h
+++ b/Engine/Plugins/MeshPainting/Source/MeshPaintEditorMode/Private/BrushSettingsCustomization.h
@@ -8,6 +8,15 @@ class FPropertyRestriction;
 class FReply;
 class IPropertyHandle;
 
+class FMeshPaintModeSettingsCustomization : public IDetailCustomization
+{
+public:
+	static TSharedRef<IDetailCustomization> MakeInstance();
+
+	/** IPropertyTypeCustomization interface */
+	virtual void CustomizeDetails(IDetailLayoutBuilder& DetailLayout) override;
+};
+
 class FMeshPaintingSettingsCustomization : public IDetailCustomization
 {
 public:
diff --git a/Engine/Plugins/MeshPainting/Source/MeshPaintEditorMode/Private/MeshPaintEditorMode.cpp b/Engine/Plugins/MeshPainting/Source/MeshPaintEditorMode/Private/MeshPaintEditorMode.cpp
--- a/Engine/Plugins/MeshPainting/Source/MeshPaintEditorMode/Private/MeshPaintEditorMode.cpp
+++ b/Engine/Plugins/MeshPainting/Source/MeshPaintEditorMode/Private/MeshPaintEditorMode.cpp
@@ -50,6 +50,7 @@ void FMeshPaintEditorModeModule::Register()
 	RegisterGeometryAdapterFactory(MakeShareable(new FMeshPaintGeometryCollectionComponentAdapterFactory));
 
 	FPropertyEditorModule& PropertyModule = FModuleManager::LoadModuleChecked<FPropertyEditorModule>("PropertyEditor");
+	PropertyModule.RegisterCustomClassLayout("MeshPaintModeSettings", FOnGetDetailCustomizationInstance::CreateStatic(&FMeshPaintModeSettingsCustomization::MakeInstance));
 	PropertyModule.RegisterCustomClassLayout("MeshVertexColorPaintingToolProperties", FOnGetDetailCustomizationInstance::CreateStatic(&FVertexColorPaintingSettingsCustomization::MakeInstance));
 	PropertyModule.RegisterCustomClassLayout("MeshVertexWeightPaintingToolProperties", FOnGetDetailCustomizationInstance::CreateStatic(&FVertexWeightPaintingSettingsCustomization::MakeInstance));
 	PropertyModule.RegisterCustomClassLayout("MeshTextureColorPaintingToolProperties", FOnGetDetailCustomizationInstance::CreateStatic(&FTextureColorPaintingSettingsCustomization::MakeInstance));
diff --git a/Engine/Plugins/MeshPainting/Source/MeshPaintEditorMode/Private/MeshPaintMode.cpp b/Engine/Plugins/MeshPainting/Source/MeshPaintEditorMode/Private/MeshPaintMode.cpp
--- a/Engine/Plugins/MeshPainting/Source/MeshPaintEditorMode/Private/MeshPaintMode.cpp
+++ b/Engine/Plugins/MeshPainting/Source/MeshPaintEditorMode/Private/MeshPaintMode.cpp
@@ -201,9 +201,9 @@ void UMeshPaintMode::CreateToolkit()
 
 void UMeshPaintMode::Tick(FEditorViewportClient* ViewportClient, float DeltaTime)
 {
-	if (bRecacheVertexDataSize)
+	if (bRecacheDataSizes)
 	{
-		UpdateCachedVertexDataSize();
+		UpdateCachedDataSizes();
 	}
 
 	if (bRecacheValidForPaint)
@@ -385,7 +385,8 @@ void UMeshPaintMode::OnVertexPaintFinished()
 		{
 			if (UMeshPaintingSubsystem* MeshPaintingSubsystem = GEngine->GetEngineSubsystem<UMeshPaintingSubsystem>())
 			{
-				UpdateCachedVertexDataSize();
+				bRecacheDataSizes = true;
+
 				MeshPaintingSubsystem->Refresh();
 			}
 		}
@@ -420,7 +421,7 @@ void UMeshPaintMode::UpdateSelectedMeshes()
 		MeshPaintingSubsystem->bNeedsRecache = true;
 	}
 	
-	bRecacheVertexDataSize = true;
+	bRecacheDataSizes = true;
 	bRecacheValidForPaint = true;
 }
 
@@ -580,6 +581,8 @@ void UMeshPaintMode::ImportVertexColorsFromFile()
 		FScopedTransaction Transaction(LOCTEXT("LevelMeshPainter_TransactionImportColors", "Importing Vertex Colors From Texture"));
 		GEditor->GetEditorSubsystem<UMeshPaintModeSubsystem>()->ImportVertexColorsFromTexture(MeshComponents[0]);
 	}
+
+	bRecacheDataSizes = true;
 }
 
 void UMeshPaintMode::SaveVertexColorsToAssets()
@@ -712,7 +715,8 @@ void UMeshPaintMode::PasteInstanceVertexColors()
 	const TArray<UStaticMeshComponent*> StaticMeshComponents = GetSelectedComponents<UStaticMeshComponent>();
 	TArray<FPerComponentVertexColorData> CopiedColorsByComponent = GEngine->GetEngineSubsystem<UMeshPaintingSubsystem>()->GetCopiedColorsByComponent();
 	GEditor->GetEditorSubsystem<UMeshPaintModeSubsystem>()->PasteVertexColors(StaticMeshComponents, CopiedColorsByComponent);
-	UpdateCachedVertexDataSize();
+	
+	bRecacheDataSizes = true;
 }
 
 void UMeshPaintMode::FixVertexColors()
@@ -723,6 +727,8 @@ void UMeshPaintMode::FixVertexColors()
 	{
 		Component->FixupOverrideColorsIfNecessary();
 	}
+
+	bRecacheDataSizes = true;
 }
 
 bool UMeshPaintMode::CanFixVertexColors() const
@@ -747,7 +753,7 @@ void UMeshPaintMode::RemoveInstanceVertexColors()
 		GEngine->GetEngineSubsystem<UMeshPaintingSubsystem>()->RemoveComponentInstanceVertexColors(Component);
 	}
 
-	UpdateCachedVertexDataSize();
+	bRecacheDataSizes = true;
 }
 
 
@@ -798,7 +804,9 @@ void UMeshPaintMode::PropagateVertexColorsToLODs()
 			FComponentReregisterContext ReregisterContext(SelectedComponent);
 		}
 	}
-	UpdateCachedVertexDataSize();
+
+	bRecacheDataSizes = true;
+	
 	MeshPaintingSubsystem->Refresh();
 }
 
@@ -847,23 +855,28 @@ template TArray<UMeshComponent*> UMeshPaintMode::GetSelectedComponents<UMeshComp
 template TArray<UGeometryCollectionComponent*> UMeshPaintMode::GetSelectedComponents<UGeometryCollectionComponent>() const;
 
 
-void UMeshPaintMode::UpdateCachedVertexDataSize()
+void UMeshPaintMode::UpdateCachedDataSizes()
 {
 	CachedVertexDataSize = 0;
+	CachedMeshPaintTextureResourceSize = 0;
 
 	const bool bInstance = true;
 	if (UMeshPaintingSubsystem* MeshPaintingSubsystem = GEngine->GetEngineSubsystem<UMeshPaintingSubsystem>())
 	{
-		for (UMeshComponent* SelectedComponent : MeshPaintingSubsystem->GetPaintableMeshComponents())
+		const TArray<UMeshComponent*> MeshComponents = GetSelectedComponents<UMeshComponent>();
+		for (UMeshComponent* MeshComponent : MeshComponents)
 		{
-			int32 NumLODs = MeshPaintingSubsystem->GetNumberOfLODs(SelectedComponent);
+			int32 NumLODs = MeshPaintingSubsystem->GetNumberOfLODs(MeshComponent);
 			for (int32 LODIndex = 0; LODIndex < NumLODs; ++LODIndex)
 			{
-				CachedVertexDataSize += MeshPaintingSubsystem->GetVertexColorBufferSize(SelectedComponent, LODIndex, bInstance);
+				CachedVertexDataSize += MeshPaintingSubsystem->GetVertexColorBufferSize(MeshComponent, LODIndex, bInstance);
 			}
+
+			CachedMeshPaintTextureResourceSize  += MeshPaintingSubsystem->GetMeshPaintTextureResourceSize(MeshComponent);
 		}
 	}
-	bRecacheVertexDataSize = false;
+
+	bRecacheDataSizes = false;
 }
 
 
@@ -1069,7 +1082,7 @@ void UMeshPaintMode::AddMeshPaintTextures()
 		GEngine->GetEngineSubsystem<UMeshPaintingSubsystem>()->CreateComponentMeshPaintTexture(Component);
 	}
 
-	// The paint status will change after add.
+	bRecacheDataSizes = true;
 	bRecacheValidForPaint = true;
 }
 
@@ -1098,7 +1111,7 @@ void UMeshPaintMode::RemoveMeshPaintTexture()
 		GEngine->GetEngineSubsystem<UMeshPaintingSubsystem>()->RemoveComponentMeshPaintTexture(Component);
 	}
 
-	// The paint status will change after removal.
+	bRecacheDataSizes = true;
 	bRecacheValidForPaint = true;
 }
 
@@ -1177,7 +1190,7 @@ void UMeshPaintMode::PasteMeshPaintTexture()
 		GEngine->GetEngineSubsystem<UMeshPaintingSubsystem>()->CreateComponentMeshPaintTexture(Component, Image);
 	}
 
-	// If this adds a mesh paint texture the selection/paint status may have changed.
+	bRecacheDataSizes = true;
 	bRecacheValidForPaint = true;
 }
 
@@ -1241,6 +1254,8 @@ void UMeshPaintMode::ImportVertexColorsFromMeshPaintTexture()
 	{
 		GEditor->GetEditorSubsystem<UMeshPaintModeSubsystem>()->ImportVertexColorsFromMeshPaintTexture(Component);
 	}
+
+	bRecacheDataSizes = true;
 }
 
 bool UMeshPaintMode::CanImportVertexColorsFromMeshPaintTexture() const
@@ -1265,7 +1280,7 @@ void UMeshPaintMode::ImportMeshPaintTextureFromVertexColors()
 		GEditor->GetEditorSubsystem<UMeshPaintModeSubsystem>()->ImportMeshPaintTextureFromVertexColors(Component);
 	}
 
-	// If this adds a mesh paint texture the selection/paint status may have changed.
+	bRecacheDataSizes = true;
 	bRecacheValidForPaint = true;
 }
 
@@ -1283,6 +1298,8 @@ void UMeshPaintMode::FixTextureColors()
 {
 	FScopedTransaction Transaction(LOCTEXT("LevelMeshPainter_TransactionFixTextureColors", "Fixing Per-Instance Texture Colors"));
 	GEditor->GetEditorSubsystem<UMeshPaintModeSubsystem>()->FixTextureColors(GetSelectedComponents<UMeshComponent>());
+
+	bRecacheDataSizes = true;
 }
 
 bool UMeshPaintMode::CanFixTextureColors() const
diff --git a/Engine/Plugins/MeshPainting/Source/MeshPaintEditorMode/Private/MeshPaintMode.h b/Engine/Plugins/MeshPainting/Source/MeshPaintEditorMode/Private/MeshPaintMode.h
--- a/Engine/Plugins/MeshPainting/Source/MeshPaintEditorMode/Private/MeshPaintMode.h
+++ b/Engine/Plugins/MeshPainting/Source/MeshPaintEditorMode/Private/MeshPaintMode.h
@@ -59,13 +59,12 @@ public:
 	template<typename ComponentClass>
 	TArray<ComponentClass*> GetSelectedComponents() const;
 
-	uint32 GetCachedVertexDataSize() const
-	{
-		return CachedVertexDataSize;
-	}
+	/** Returns data size of per-instance vertex color data for the currently selected components. */
+	uint32 GetVertexDataSizeInBytes() const { return CachedVertexDataSize; }
+	/** Returns resource size of mesh paint textures for the currently selected components. */
+	uint32 GetMeshPaintTextureResourceSizeInBytes() const { return CachedMeshPaintTextureResourceSize; }
 
 protected:
-
 	/** Binds UI commands to actions for the mesh paint mode */
 	virtual void BindCommands() override;
 
@@ -85,7 +84,7 @@ protected:
 	void OnResetViewMode();
 	void OnVertexPaintFinished();
 
-	void UpdateCachedVertexDataSize();
+	void UpdateCachedDataSizes();
 	void EndPaintToolIfNoLongerValid();
 
 	bool IsInSelectTool() const;
@@ -146,8 +145,9 @@ protected:
 	UPROPERTY(Transient)
 	TObjectPtr<UMeshPaintModeSettings> ModeSettings;
 
+	bool bRecacheDataSizes = false;
 	uint32 CachedVertexDataSize = 0;
-	bool bRecacheVertexDataSize = false;
+	uint32 CachedMeshPaintTextureResourceSize = 0;
 	
 	bool bRecacheValidForPaint = false;
 
diff --git a/Engine/Plugins/MeshPainting/Source/MeshPaintEditorMode/Private/MeshPaintModeSettings.h b/Engine/Plugins/MeshPainting/Source/MeshPaintEditorMode/Private/MeshPaintModeSettings.h
--- a/Engine/Plugins/MeshPainting/Source/MeshPaintEditorMode/Private/MeshPaintModeSettings.h
+++ b/Engine/Plugins/MeshPainting/Source/MeshPaintEditorMode/Private/MeshPaintModeSettings.h
@@ -6,7 +6,6 @@
 #include "UObject/Package.h"
 #include "MeshPaintModeSettings.generated.h"
 
-
 /** Mesh paint color view modes (somewhat maps to EVertexColorViewMode engine enum.) */
 UENUM()
 enum class EMeshPaintColorView : uint8
@@ -40,17 +39,7 @@ class MESHPAINTEDITORMODE_API UMeshPaintModeSettings
 	GENERATED_BODY()
 
 public:
-
-	UMeshPaintModeSettings()
-	{
-		ColorViewMode = EMeshPaintDataColorViewMode::Normal;
-	}
-
-public:
-
 	/** Color view mode used to display Vertex Colors */
 	UPROPERTY(config, EditAnywhere, Category = Visualization)
-	EMeshPaintDataColorViewMode ColorViewMode;
-
-
+	EMeshPaintDataColorViewMode ColorViewMode = EMeshPaintDataColorViewMode::Normal;
 };
diff --git a/Engine/Plugins/MeshPainting/Source/MeshPaintingToolset/Private/MeshPaintHelpers.cpp b/Engine/Plugins/MeshPainting/Source/MeshPaintingToolset/Private/MeshPaintHelpers.cpp
--- a/Engine/Plugins/MeshPainting/Source/MeshPaintingToolset/Private/MeshPaintHelpers.cpp
+++ b/Engine/Plugins/MeshPainting/Source/MeshPaintingToolset/Private/MeshPaintHelpers.cpp
@@ -643,6 +643,23 @@ uint32 UMeshPaintingSubsystem::GetVertexColorBufferSize(UMeshComponent* MeshComp
 	return SizeInBytes;
 }
 
+uint32 UMeshPaintingSubsystem::GetMeshPaintTextureResourceSize(UMeshComponent* MeshComponent)
+{
+	if (UTexture* Texture = MeshComponent->GetMeshPaintTexture())
+	{
+		FTexturePlatformData** PlatformDataPtr = Texture->GetRunningPlatformData();
+		if (PlatformDataPtr != nullptr && *PlatformDataPtr != nullptr)
+		{
+			return (*PlatformDataPtr)->GetPayloadSize(0);
+		}
+		else
+		{
+			return Texture->GetResourceSizeBytes(EResourceSizeMode::Exclusive);
+		}
+	}
+	return 0;
+}
+
 TArray<FVector> UMeshPaintingSubsystem::GetVerticesForLOD( const UStaticMesh* StaticMesh, int32 LODIndex)
 {
 	checkf(StaticMesh != nullptr, TEXT("Invalid static mesh ptr"));
diff --git a/Engine/Plugins/MeshPainting/Source/MeshPaintingToolset/Public/MeshPaintHelpers.h b/Engine/Plugins/MeshPainting/Source/MeshPaintingToolset/Public/MeshPaintHelpers.h
--- a/Engine/Plugins/MeshPainting/Source/MeshPaintingToolset/Public/MeshPaintHelpers.h
+++ b/Engine/Plugins/MeshPainting/Source/MeshPaintingToolset/Public/MeshPaintHelpers.h
@@ -149,6 +149,9 @@ public:
 
 	/** Retrieves the Vertex Color buffer size for the given LOD level in the Mesh */
 	uint32 GetVertexColorBufferSize(UMeshComponent* MeshComponent, int32 LODIndex, bool bInstance);
+	
+	/** Retrieves the resource size for the mesh paint texture on the component */
+	uint32 GetMeshPaintTextureResourceSize(UMeshComponent* MeshComponent);
 
 	/** Retrieves the vertex positions from the given LOD level in the Mesh */
 	TArray<FVector> GetVerticesForLOD(const UStaticMesh* StaticMesh, int32 LODIndex);
diff --git a/Engine/Plugins/MeshPainting/Source/MeshPaintingToolset/Public/MeshVertexPaintingTool.h b/Engine/Plugins/MeshPainting/Source/MeshPaintingToolset/Public/MeshVertexPaintingTool.h
--- a/Engine/Plugins/MeshPainting/Source/MeshPaintingToolset/Public/MeshVertexPaintingTool.h
+++ b/Engine/Plugins/MeshPainting/Source/MeshPaintingToolset/Public/MeshVertexPaintingTool.h
@@ -82,7 +82,7 @@ public:
 	UMeshVertexPaintingToolProperties();
 
 	/** Size of vertex points drawn when mesh painting is active. */
-	UPROPERTY(EditAnywhere, Category = "VertexPainting|Visualization")
+	UPROPERTY(EditAnywhere, AdvancedDisplay, Category = "VertexPainting")
 	float VertexPreviewSize;
 };
 
