Fix texture color paint copy buffer being cleared when not wanted just by moving between mesh paint tools.

diff --git a/Engine/Plugins/MeshPainting/Source/MeshPaintingToolset/Private/MeshPaintHelpers.cpp b/Engine/Plugins/MeshPainting/Source/MeshPaintingToolset/Private/MeshPaintHelpers.cpp
--- a/Engine/Plugins/MeshPainting/Source/MeshPaintingToolset/Private/MeshPaintHelpers.cpp
+++ b/Engine/Plugins/MeshPainting/Source/MeshPaintingToolset/Private/MeshPaintHelpers.cpp
@@ -1553,8 +1553,6 @@ void UMeshPaintingSubsystem::CleanUp()
 	}
 	ComponentToAdapterMap.Empty();
 	FMeshPaintComponentAdapterFactory::CleanupGlobals();
-
-	CopiedTextureData.Init(FImageInfo());
 }
 
 bool UMeshPaintingSubsystem::FindHitResult(const FRay Ray, FHitResult& BestTraceResult)
