Change logic for enabling tools during texture mesh painting.
Previously we conservatively disabled some things like copy/paste/fix etc unless we were in the select tool.
Now we enable those in the paint tool too, but only if there are no outstanding paint strokes to be applied.

diff --git a/Engine/Plugins/MeshPainting/Source/MeshPaintEditorMode/Private/MeshPaintMode.cpp b/Engine/Plugins/MeshPainting/Source/MeshPaintEditorMode/Private/MeshPaintMode.cpp
--- a/Engine/Plugins/MeshPainting/Source/MeshPaintEditorMode/Private/MeshPaintMode.cpp
+++ b/Engine/Plugins/MeshPainting/Source/MeshPaintEditorMode/Private/MeshPaintMode.cpp
@@ -1065,6 +1065,10 @@ void UMeshPaintMode::SaveTexturePackages()
 
 bool UMeshPaintMode::CanSaveTexturePackages() const
 {
+	if (CanApplyTextureColorsToAsset())
+	{
+		return false;
+	}
 	TArray<UObject*> TexturesToSave;
 	if (UMeshTexturePaintingTool* TexturePaintingTool = Cast<UMeshTexturePaintingTool>(GetToolManager()->GetActiveTool(EToolSide::Left)))
 	{
@@ -1117,7 +1121,7 @@ void UMeshPaintMode::RemoveMeshPaintTexture()
 
 bool UMeshPaintMode::CanRemoveMeshPaintTextures() const
 {
-	if (IsInSelectTool())
+	if (!CanApplyTextureColorsToAsset())
 	{
 		const TArray<UStaticMeshComponent*> StaticMeshComponents = GetSelectedComponents<UStaticMeshComponent>();
 		for (UStaticMeshComponent* Component : StaticMeshComponents)
@@ -1142,7 +1146,7 @@ void UMeshPaintMode::CopyMeshPaintTexture()
 
 bool UMeshPaintMode::CanCopyMeshPaintTexture() const
 {
-	if (IsInSelectTool())
+	if (!CanApplyTextureColorsToAsset())
 	{
 		const TArray<UStaticMeshComponent*> StaticMeshComponents = GetSelectedComponents<UStaticMeshComponent>();
 		if (StaticMeshComponents.Num() == 1 && StaticMeshComponents[0]->GetMeshPaintTexture() != nullptr)
@@ -1196,7 +1200,7 @@ void UMeshPaintMode::PasteMeshPaintTexture()
 
 bool UMeshPaintMode::CanPasteMeshPaintTexture() const
 {
-	if (!IsInSelectTool())
+	if (CanApplyTextureColorsToAsset())
 	{
 		return false;
 	}
@@ -1286,7 +1290,7 @@ void UMeshPaintMode::ImportMeshPaintTextureFromVertexColors()
 
 bool UMeshPaintMode::CanImportMeshPaintTextureFromVertexColors() const
 {
-	if (IsInSelectTool())
+	if (!CanApplyTextureColorsToAsset())
 	{
 		const TArray<UMeshComponent*> MeshComponents = GetSelectedComponents<UMeshComponent>();
 		return MeshComponents.Num() > 0;
@@ -1304,7 +1308,7 @@ void UMeshPaintMode::FixTextureColors()
 
 bool UMeshPaintMode::CanFixTextureColors() const
 {
-	if (!IsInSelectTool())
+	if (CanApplyTextureColorsToAsset())
 	{
 		return false;
 	}
