Fix mesh paint default UV channel not being inferrd from the material.

diff --git a/Engine/Plugins/MeshPainting/Source/MeshPaintingToolset/Private/IMeshPaintComponentAdapter.cpp b/Engine/Plugins/MeshPainting/Source/MeshPaintingToolset/Private/IMeshPaintComponentAdapter.cpp
--- a/Engine/Plugins/MeshPainting/Source/MeshPaintingToolset/Private/IMeshPaintComponentAdapter.cpp
+++ b/Engine/Plugins/MeshPainting/Source/MeshPaintingToolset/Private/IMeshPaintComponentAdapter.cpp
@@ -87,6 +87,10 @@ void IMeshPaintComponentAdapter::DefaultQueryPaintableTextures(int32 MaterialInd
 							// Store the uv channel, this is set when the texture is selected. 
 							PaintableTexture.UVChannelIndex = TextureCoords->CoordinateIndex;
 						}
+						else
+						{
+							PaintableTexture.UVChannelIndex = TextureSample->ConstCoordinate;
+						}
 
 						// Handle texture parameter expressions
 						UMaterialExpressionTextureSampleParameter* TextureSampleParameter = Cast<UMaterialExpressionTextureSampleParameter>(TextureSample);
diff --git a/Engine/Plugins/MeshPainting/Source/MeshPaintingToolset/Private/MeshTexturePaintingTool.cpp b/Engine/Plugins/MeshPainting/Source/MeshPaintingToolset/Private/MeshTexturePaintingTool.cpp
--- a/Engine/Plugins/MeshPainting/Source/MeshPaintingToolset/Private/MeshTexturePaintingTool.cpp
+++ b/Engine/Plugins/MeshPainting/Source/MeshPaintingToolset/Private/MeshTexturePaintingTool.cpp
@@ -203,8 +203,6 @@ void UMeshTexturePaintingTool::Shutdown(EToolShutdownType ShutdownType)
 
 	PaintTargetData.Empty();
 
-	// Remove any existing texture targets
-	TexturePaintTargetList.Empty();
 	if (UMeshPaintingSubsystem* MeshPaintingSubsystem = GEngine->GetEngineSubsystem<UMeshPaintingSubsystem>())
 	{
 		MeshPaintingSubsystem->Refresh();
@@ -1825,21 +1823,14 @@ bool UMeshTextureAssetPaintingTool::ShouldFilterTextureAsset(const FAssetData& A
 
 void UMeshTextureAssetPaintingTool::PaintTextureChanged(const FAssetData& AssetData)
 {
-	UTexture2D* Texture = Cast<UTexture2D>(AssetData.GetAsset());
-	if (Texture)
+	if (UTexture2D* Texture = Cast<UTexture2D>(AssetData.GetAsset()))
 	{
-		// Loop through our list of textures and see which one the user wants to select
-		for (int32 TargetIndex = 0; TargetIndex < TexturePaintTargetList.Num(); TargetIndex++)
+		for (FPaintableTexture& PaintableTexture : PaintableTextures)
 		{
-			FTextureTargetListInfo& TextureTarget = TexturePaintTargetList[TargetIndex];
-			if (TextureTarget.TextureData == Texture)
+			if (PaintableTexture.Texture == Texture)
 			{
-				TextureTarget.bIsSelected = true;
-				AssetProperties->UVChannel = TextureTarget.UVChannelIndex;
-			}
-			else
-			{
-				TextureTarget.bIsSelected = false;
+				AssetProperties->UVChannel = PaintableTexture.UVChannelIndex;
+				break;
 			}
 		}
 	}
@@ -1853,6 +1844,7 @@ void UMeshTextureAssetPaintingTool::CacheTexturePaintData()
 		PaintableTextures.Empty();
 		
 		UTexture* DefaultTexture = nullptr;
+		int32 DefaultUVChannelIndex = INDEX_NONE;
 	
 		TArray<UMeshComponent*> PaintableComponents = MeshPaintingSubsystem->GetPaintableMeshComponents();
 		if (PaintableComponents.IsValidIndex(0))
@@ -1863,6 +1855,7 @@ void UMeshTextureAssetPaintingTool::CacheTexturePaintData()
 			if (PaintableTextures.IsValidIndex(DefaultTextureIndex))
 			{
 				DefaultTexture = PaintableTextures[DefaultTextureIndex].Texture;
+				DefaultUVChannelIndex = PaintableTextures[DefaultTextureIndex].UVChannelIndex;
 			}
 		}
 
@@ -1872,15 +1865,19 @@ void UMeshTextureAssetPaintingTool::CacheTexturePaintData()
 		if (!PaintableTextures.Contains(AssetProperties->PaintTexture))
 		{
 			UTexture2D* NewTexture = nullptr;
+			int32 NewUVChannelIndex = INDEX_NONE;
 			if (PaintableTextures.Contains(DefaultTexture))
 			{
 				NewTexture = Cast<UTexture2D>(DefaultTexture);
+				NewUVChannelIndex = DefaultUVChannelIndex;
 			}
 			else if (PaintableTextures.Num() > 0)
 			{
 				NewTexture = Cast<UTexture2D>(PaintableTextures[0].Texture);
+				NewUVChannelIndex = PaintableTextures[0].UVChannelIndex;
 			}
 			AssetProperties->PaintTexture = NewTexture;
+			AssetProperties->UVChannel = NewUVChannelIndex;
 		}
 	}
 }
diff --git a/Engine/Plugins/MeshPainting/Source/MeshPaintingToolset/Public/MeshTexturePaintingTool.h b/Engine/Plugins/MeshPainting/Source/MeshPaintingToolset/Public/MeshTexturePaintingTool.h
--- a/Engine/Plugins/MeshPainting/Source/MeshPaintingToolset/Public/MeshTexturePaintingTool.h
+++ b/Engine/Plugins/MeshPainting/Source/MeshPaintingToolset/Public/MeshTexturePaintingTool.h
@@ -224,9 +224,6 @@ protected:
 	UPROPERTY(Transient, NonTransactional)
 	TMap<TObjectPtr<UTexture2D>, FPaintComponentOverride> PaintComponentsOverride;
 
-	/** Texture paint: Will hold a list of texture items that we can paint on */
-	TArray<FTextureTargetListInfo> TexturePaintTargetList;
-
 	/** Texture paint: The mesh components that we're currently painting */
 	UPROPERTY(Transient)
 	TObjectPtr<UMeshComponent> TexturePaintingCurrentMeshComponent;
