Rewrite of undo approach for mesh texture painting.
We now immediately write brush strokes to the final texture and rely on the texture transaction to handle undo/redo.
This is much both simpler to implement and more intuitive to the user than the previous approach where we had intermediate brush strokes and a final apply to write to texture.
Removed the apply button for texture painting (and updated the name for the vertex apply/propagate to mesh while I was there).
Also painting on multiple components in a single stroke is now a single transaction instead of multiple transactions.
One thing to note is that we now show the render target of paintable textures always when in texture painting.
Previously we only showed them after the first brush stroke and before the Apply.
The change is to reduce as much as possible any flicker/stall that would result from waiting on the paint texture compile at the end of each brush stroke.
For compressed virtual textures (like the mesh paint textures) there can be a visible difference between the render target during painting and the final result.
This is becuase we use a virtual texture adapter with GPU runtime compression.
The GPU compression is lower quality than the final CPU texture compilation compression. Common artifacts are desaturated colors where primaries blend across a compression block.
This can look worse for very low resolution paint textures, and may need to be revisited (maybe we can paint at higher res or something).
Also fixed bug where non-sRGB textures weren't being read back from the render target correctly leading to "hardening" of painted edges over time.

diff --git a/Engine/Plugins/MeshPainting/Source/MeshPaintEditorMode/Private/BrushSettingsCustomization.cpp b/Engine/Plugins/MeshPainting/Source/MeshPaintEditorMode/Private/BrushSettingsCustomization.cpp
--- a/Engine/Plugins/MeshPainting/Source/MeshPaintEditorMode/Private/BrushSettingsCustomization.cpp
+++ b/Engine/Plugins/MeshPainting/Source/MeshPaintEditorMode/Private/BrushSettingsCustomization.cpp
@@ -14,6 +14,7 @@
 #include "InteractiveToolManager.h"
 #include "MeshPaintMode.h"
 #include "Components/SkeletalMeshComponent.h"
+#include "MeshPaintHelpers.h"
 #include "MeshPaintModeHelpers.h"
 #include "MeshTexturePaintingTool.h"
 #include "ContentBrowserDelegates.h"
@@ -522,7 +523,6 @@ void FTextureAssetPaintingSettingsCustomization::CustomizeDetails(IDetailLayoutB
 		.PropertyHandle(TextureProperty)
 		.AllowedClass(UTexture2D::StaticClass())
 		.OnShouldFilterAsset(FOnShouldFilterAsset::CreateUObject(Cast<UMeshTextureAssetPaintingTool>(UMeshPaintMode::GetMeshPaintMode()->GetToolManager()->GetActiveTool(EToolSide::Left)), &UMeshTextureAssetPaintingTool::ShouldFilterTextureAsset))
-		.OnObjectChanged(FOnSetObject::CreateUObject(Cast<UMeshTextureAssetPaintingTool>(UMeshPaintMode::GetMeshPaintMode()->GetToolManager()->GetActiveTool(EToolSide::Left)), &UMeshTextureAssetPaintingTool::PaintTextureChanged))
 		.DisplayUseSelected(false)
 		.ThumbnailPool(DetailLayout.GetThumbnailPool())
 	];
diff --git a/Engine/Plugins/MeshPainting/Source/MeshPaintEditorMode/Private/MeshPaintMode.cpp b/Engine/Plugins/MeshPainting/Source/MeshPaintEditorMode/Private/MeshPaintMode.cpp
--- a/Engine/Plugins/MeshPainting/Source/MeshPaintEditorMode/Private/MeshPaintMode.cpp
+++ b/Engine/Plugins/MeshPainting/Source/MeshPaintEditorMode/Private/MeshPaintMode.cpp
@@ -18,6 +18,7 @@
 #include "ScopedTransaction.h"
 #include "PackageTools.h"
 #include "MeshPaintAdapterFactory.h"
+#include "MeshPaintHelpers.h"
 #include "MeshPaintModeHelpers.h"
 #include "MeshSelect.h"
 #include "MeshTexturePaintingTool.h"
@@ -154,10 +155,11 @@ void UMeshPaintMode::Enter()
 	// Toolkit
 	PaletteChangedHandle = Toolkit->OnPaletteChanged().AddUObject(this, &UMeshPaintMode::UpdateOnPaletteChange);
 
-	// disable tool change tracking to activate default tool, and then switch to full undo/redo tracking mode
+	// disable tool change tracking to activate default tool
 	GetToolManager()->ConfigureChangeTrackingMode(EToolChangeTrackingMode::NoChangeTracking);
 	Toolkit->SetCurrentPalette(MeshPaintMode_VertexColor);
-	GetToolManager()->ConfigureChangeTrackingMode(EToolChangeTrackingMode::UndoToExit);
+	// switch back to full undo / redo tracking mode here if that is behavior we want
+	//GetToolManager()->ConfigureChangeTrackingMode(EToolChangeTrackingMode::FullUndoRedo);
 
 	FLevelEditorModule& LevelEditor = FModuleManager::GetModuleChecked<FLevelEditorModule>(FName(TEXT("LevelEditor")));
 	LevelEditor.OnRedrawLevelEditingViewports().AddUObject(this, &UMeshPaintMode::UpdateOnMaterialChange);
@@ -286,15 +288,11 @@ void UMeshPaintMode::BindCommands()
  		FExecuteAction::CreateUObject(this, &UMeshPaintMode::FillTexture),
  		FCanExecuteAction::CreateUObject(this, &UMeshPaintMode::CanFillTexture)));
 
- 	CommandList->MapAction(Commands.ApplyVertex, FUIAction(
- 		FExecuteAction::CreateUObject(this, &UMeshPaintMode::ApplyVertexColorsToAsset),
- 		FCanExecuteAction::CreateUObject(this, &UMeshPaintMode::CanApplyVertexColorsToAsset)));
+ 	CommandList->MapAction(Commands.PropagateMesh, FUIAction(
+ 		FExecuteAction::CreateUObject(this, &UMeshPaintMode::PropagateVertexColorsToMesh),
+ 		FCanExecuteAction::CreateUObject(this, &UMeshPaintMode::CanPropagateVertexColorsToMesh)));
 
- 	CommandList->MapAction(Commands.ApplyTexture, FUIAction(
- 		FExecuteAction::CreateUObject(this, &UMeshPaintMode::ApplyTextureColorsToAsset),
- 		FCanExecuteAction::CreateUObject(this, &UMeshPaintMode::CanApplyTextureColorsToAsset)));
-
-	CommandList->MapAction(Commands.Propagate, FUIAction(
+	CommandList->MapAction(Commands.PropagateLODs, FUIAction(
 		FExecuteAction::CreateUObject(this, &UMeshPaintMode::PropagateVertexColorsToLODs),
 		FCanExecuteAction::CreateUObject(this, &UMeshPaintMode::CanPropagateVertexColorsToLODs)));
 
@@ -371,6 +369,11 @@ void UMeshPaintMode::OnToolStarted(UInteractiveToolManager* Manager, UInteractiv
 	{
 		VertexPaintingTool->OnPaintingFinished().BindUObject(this, &UMeshPaintMode::OnVertexPaintFinished);
 	}
+
+	if (UMeshTextureColorPaintingTool* TextureColorPaintingTool = Cast<UMeshTextureColorPaintingTool>(GetToolManager()->GetActiveTool(EToolSide::Left)))
+	{
+		TextureColorPaintingTool->OnPaintingFinished().BindUObject(this, &UMeshPaintMode::OnTextureColorVertexPaintFinished);
+	}
 }
 
 void UMeshPaintMode::OnVertexPaintFinished()
@@ -393,6 +396,17 @@ void UMeshPaintMode::OnVertexPaintFinished()
 	}
 }
 
+void UMeshPaintMode::OnTextureColorVertexPaintFinished(UMeshComponent* MeshComponent)
+{
+	if (UMeshTextureColorPaintingToolProperties* TexturePaintingToolProperties = UMeshPaintMode::GetTextureColorToolProperties())
+	{
+		if (TexturePaintingToolProperties->bPropagateToVertexColor)
+		{
+			GEditor->GetEditorSubsystem<UMeshPaintModeSubsystem>()->ImportVertexColorsFromMeshPaintTexture(MeshComponent);
+		}
+	}
+}
+
 void UMeshPaintMode::OnToolEnded(UInteractiveToolManager* Manager, UInteractiveTool* Tool)
 {
 	FMeshPaintingToolActionCommands::UpdateToolCommandBinding(Tool, Toolkit->GetToolkitCommands(), true);
@@ -541,7 +555,7 @@ void UMeshPaintMode::FillVertexColors()
 }
 
 
-void UMeshPaintMode::ApplyVertexColorsToAsset()
+void UMeshPaintMode::PropagateVertexColorsToMesh()
 {
 	const TArray<UStaticMeshComponent*> StaticMeshComponents = GetSelectedComponents<UStaticMeshComponent>();
 	FSuppressableWarningDialog::FSetupInfo SetupInfo(LOCTEXT("PushInstanceVertexColorsPrompt_Message", "This operation copies vertex colors from LOD 0 of the selected instance to all LODs of the source asset, overwriting any existing vertex colors.\n\nThis change will also propagate to all other instances of the same asset that do not have custom vertex colors."),
@@ -562,7 +576,7 @@ void UMeshPaintMode::ApplyVertexColorsToAsset()
 	}
 }
 
-bool UMeshPaintMode::CanApplyVertexColorsToAsset() const
+bool UMeshPaintMode::CanPropagateVertexColorsToMesh() const
 {
 	// Check whether or not our selected Static Mesh Components contain instance based vertex colors (only these can be propagated to the base mesh)
 	int32 NumInstanceVertexColorBytes = 0;
@@ -905,36 +919,6 @@ bool UMeshPaintMode::CanCycleTextures() const
 	return TexturePaintingTool != nullptr;
 }
 
-void UMeshPaintMode::ApplyTextureColorsToAsset()
-{
-	FScopedTransaction Transaction(LOCTEXT("MeshPaintMode_TexturePaint_Apply", "Apply Texture Paint"));
-
-	TArray<UMeshComponent*> VertexColorPropagateComponents;
-	if (UMeshTextureColorPaintingTool* TextureColorPaintingTool = Cast<UMeshTextureColorPaintingTool>(GetToolManager()->GetActiveTool(EToolSide::Left)))
-	{
-		TextureColorPaintingTool->GetPaintedComponentsForPropagateVertexColor(VertexColorPropagateComponents);
-	}
-
-	if (UMeshTexturePaintingTool* TexturePaintingTool = Cast<UMeshTexturePaintingTool>(GetToolManager()->GetActiveTool(EToolSide::Left)))
-	{
-		TexturePaintingTool->ApplyAllPaintedTextures();
-	}
-
-	for (UMeshComponent* Component : VertexColorPropagateComponents)
-	{
-		GEditor->GetEditorSubsystem<UMeshPaintModeSubsystem>()->ImportVertexColorsFromMeshPaintTexture(Component);
-	}
-}
-
-bool UMeshPaintMode::CanApplyTextureColorsToAsset() const
-{
-	if (UMeshTexturePaintingTool* TexturePaintingTool = Cast<UMeshTexturePaintingTool>(GetToolManager()->GetActiveTool(EToolSide::Left)))
-	{
-		return TexturePaintingTool->GetNumberOfPendingPaintChanges() > 0;
-	}
-	return false;
-}
-
 
 void UMeshPaintMode::ActivateDefaultTool()
 {
@@ -1065,10 +1049,6 @@ void UMeshPaintMode::SaveTexturePackages()
 
 bool UMeshPaintMode::CanSaveTexturePackages() const
 {
-	if (CanApplyTextureColorsToAsset())
-	{
-		return false;
-	}
 	TArray<UObject*> TexturesToSave;
 	if (UMeshTexturePaintingTool* TexturePaintingTool = Cast<UMeshTexturePaintingTool>(GetToolManager()->GetActiveTool(EToolSide::Left)))
 	{
@@ -1121,15 +1101,12 @@ void UMeshPaintMode::RemoveMeshPaintTexture()
 
 bool UMeshPaintMode::CanRemoveMeshPaintTextures() const
 {
-	if (!CanApplyTextureColorsToAsset())
+	const TArray<UStaticMeshComponent*> StaticMeshComponents = GetSelectedComponents<UStaticMeshComponent>();
+	for (UStaticMeshComponent* Component : StaticMeshComponents)
 	{
-		const TArray<UStaticMeshComponent*> StaticMeshComponents = GetSelectedComponents<UStaticMeshComponent>();
-		for (UStaticMeshComponent* Component : StaticMeshComponents)
+		if (Component->GetMeshPaintTexture() != nullptr)
 		{
-			if (Component->GetMeshPaintTexture() != nullptr)
-			{
-				return true;
-			}
+			return true;
 		}
 	}
 	return false;
@@ -1146,13 +1123,10 @@ void UMeshPaintMode::CopyMeshPaintTexture()
 
 bool UMeshPaintMode::CanCopyMeshPaintTexture() const
 {
-	if (!CanApplyTextureColorsToAsset())
+	const TArray<UStaticMeshComponent*> StaticMeshComponents = GetSelectedComponents<UStaticMeshComponent>();
+	if (StaticMeshComponents.Num() == 1 && StaticMeshComponents[0]->GetMeshPaintTexture() != nullptr)
 	{
-		const TArray<UStaticMeshComponent*> StaticMeshComponents = GetSelectedComponents<UStaticMeshComponent>();
-		if (StaticMeshComponents.Num() == 1 && StaticMeshComponents[0]->GetMeshPaintTexture() != nullptr)
-		{
-			return true;
-		}
+		return true;
 	}
 	return false;
 }
@@ -1196,14 +1170,12 @@ void UMeshPaintMode::PasteMeshPaintTexture()
 
 	bRecacheDataSizes = true;
 	bRecacheValidForPaint = true;
+	
+	GEngine->GetEngineSubsystem<UMeshPaintingSubsystem>()->Refresh();
 }
 
 bool UMeshPaintMode::CanPasteMeshPaintTexture() const
 {
-	if (CanApplyTextureColorsToAsset())
-	{
-		return false;
-	}
 	const TArray<UStaticMeshComponent*> StaticMeshComponents = GetSelectedComponents<UStaticMeshComponent>();
 	if (StaticMeshComponents.Num() == 0)
 	{
@@ -1286,16 +1258,14 @@ void UMeshPaintMode::ImportMeshPaintTextureFromVertexColors()
 
 	bRecacheDataSizes = true;
 	bRecacheValidForPaint = true;
+
+	GEngine->GetEngineSubsystem<UMeshPaintingSubsystem>()->Refresh();
 }
 
 bool UMeshPaintMode::CanImportMeshPaintTextureFromVertexColors() const
 {
-	if (!CanApplyTextureColorsToAsset())
-	{
-		const TArray<UMeshComponent*> MeshComponents = GetSelectedComponents<UMeshComponent>();
-		return MeshComponents.Num() > 0;
-	}
-	return false;
+	const TArray<UStaticMeshComponent*> MeshComponents = GetSelectedComponents<UStaticMeshComponent>();
+	return MeshComponents.Num() > 0;
 }
 
 void UMeshPaintMode::FixTextureColors()
@@ -1304,15 +1274,12 @@ void UMeshPaintMode::FixTextureColors()
 	GEditor->GetEditorSubsystem<UMeshPaintModeSubsystem>()->FixTextureColors(GetSelectedComponents<UMeshComponent>());
 
 	bRecacheDataSizes = true;
+
+	GEngine->GetEngineSubsystem<UMeshPaintingSubsystem>()->Refresh();
 }
 
 bool UMeshPaintMode::CanFixTextureColors() const
 {
-	if (CanApplyTextureColorsToAsset())
-	{
-		return false;
-	}
-
 	return GEditor->GetEditorSubsystem<UMeshPaintModeSubsystem>()->CanFixTextureColors(GetSelectedComponents<UMeshComponent>());
 }
 
diff --git a/Engine/Plugins/MeshPainting/Source/MeshPaintEditorMode/Private/MeshPaintMode.h b/Engine/Plugins/MeshPainting/Source/MeshPaintEditorMode/Private/MeshPaintMode.h
--- a/Engine/Plugins/MeshPainting/Source/MeshPaintEditorMode/Private/MeshPaintMode.h
+++ b/Engine/Plugins/MeshPainting/Source/MeshPaintEditorMode/Private/MeshPaintMode.h
@@ -27,8 +27,15 @@ class UMeshPaintMode : public UEdMode, public ILegacyEdModeViewportInterface
 public:
 	GENERATED_BODY()
 
-	/** Default constructor for UMeshPaintMode */
 	UMeshPaintMode();
+
+	virtual void Enter() override;
+	virtual void Exit() override;
+	virtual void CreateToolkit() override;
+	virtual void Tick(FEditorViewportClient* ViewportClient, float DeltaTime) override;
+	virtual bool HandleClick(FEditorViewportClient* InViewportClient, HHitProxy* HitProxy, const FViewportClick& Click) override;
+	virtual TMap<FName, TArray<TSharedPtr<FUICommandInfo>>> GetModeCommands() const override;
+
 	static UMeshPaintingToolProperties* GetToolProperties();
 	static UMeshVertexPaintingToolProperties* GetVertexToolProperties();
 	static UMeshVertexColorPaintingToolProperties* GetVertexColorToolProperties();
@@ -37,11 +44,7 @@ public:
 	static UMeshTextureColorPaintingToolProperties* GetTextureColorToolProperties();
 	static UMeshTextureAssetPaintingToolProperties* GetTextureAssetToolProperties();
 	static UMeshPaintMode* GetMeshPaintMode();
-	virtual void Enter() override;
-	virtual void Exit() override;
-	virtual void CreateToolkit() override;
-	virtual void Tick(FEditorViewportClient* ViewportClient, float DeltaTime) override;
-	virtual bool HandleClick(FEditorViewportClient* InViewportClient, HHitProxy* HitProxy, const FViewportClick& Click) override;
+
 	static FName MeshPaintMode_VertexColor;
 	static FName MeshPaintMode_VertexWeights;
 	static FName MeshPaintMode_TextureColor;
@@ -54,7 +57,6 @@ public:
 	static FString TextureColorPaintToolName;
 	static FString TextureAssetPaintToolName;
 
-	virtual TMap<FName, TArray<TSharedPtr<FUICommandInfo>>> GetModeCommands() const override;
 	/** Returns the instance of ComponentClass found in the current Editor selection */
 	template<typename ComponentClass>
 	TArray<ComponentClass*> GetSelectedComponents() const;
@@ -83,6 +85,7 @@ protected:
 	void OnObjectsReplaced(const TMap<UObject*, UObject*>& InOldToNewInstanceMap);
 	void OnResetViewMode();
 	void OnVertexPaintFinished();
+	void OnTextureColorVertexPaintFinished(UMeshComponent* MeshComponent);
 
 	void UpdateCachedDataSizes();
 	void EndPaintToolIfNoLongerValid();
@@ -97,10 +100,8 @@ protected:
 	bool CanFillVertexColors() const;
 	void FillTexture();
 	bool CanFillTexture() const;
-	void ApplyVertexColorsToAsset();
-	bool CanApplyVertexColorsToAsset() const;
-	void ApplyTextureColorsToAsset();
-	bool CanApplyTextureColorsToAsset() const;
+	void PropagateVertexColorsToMesh();
+	bool CanPropagateVertexColorsToMesh() const;
 	void PropagateVertexColorsToLODs();
 	bool CanPropagateVertexColorsToLODs() const;
 	void SaveVertexColorsToAssets();
diff --git a/Engine/Plugins/MeshPainting/Source/MeshPaintEditorMode/Private/MeshPaintModeCommands.cpp b/Engine/Plugins/MeshPainting/Source/MeshPaintEditorMode/Private/MeshPaintModeCommands.cpp
--- a/Engine/Plugins/MeshPainting/Source/MeshPaintEditorMode/Private/MeshPaintModeCommands.cpp
+++ b/Engine/Plugins/MeshPainting/Source/MeshPaintEditorMode/Private/MeshPaintModeCommands.cpp
@@ -41,9 +41,8 @@ void FMeshPaintEditorModeCommands::RegisterCommands()
 	UI_COMMAND(SwapColor, "Swap", "Switches the foreground and background colors used for painting", EUserInterfaceActionType::Button, FInputChord(EKeys::X));
 	UI_COMMAND(FillVertex, "Fill", "Fills the selected meshes with the paint color", EUserInterfaceActionType::Button, FInputChord());
 	UI_COMMAND(FillTexture, "Fill", "Fills the selected textures with the paint color", EUserInterfaceActionType::Button, FInputChord());
-	UI_COMMAND(ApplyVertex, "Apply", "Applies per instance vertex colors to the source meshes", EUserInterfaceActionType::Button, FInputChord());
-	UI_COMMAND(ApplyTexture, "Apply", "Applies texture edits to the source texture assets", EUserInterfaceActionType::Button, FInputChord(EKeys::C, EModifierKey::Control | EModifierKey::Shift));
-	UI_COMMAND(Propagate, "All LODs", "Propagates the vertex colors from LOD0 to all LOD levels", EUserInterfaceActionType::Button, FInputChord());
+	UI_COMMAND(PropagateMesh, "ToMesh", "Applies per instance vertex colors to the source meshes", EUserInterfaceActionType::Button, FInputChord());
+	UI_COMMAND(PropagateLODs, "ToLODs", "Applies the vertex colors from LOD0 to all LOD levels", EUserInterfaceActionType::Button, FInputChord());
 	UI_COMMAND(SaveVertex, "Save", "Saves the source meshes for the selected mesh components", EUserInterfaceActionType::Button, FInputChord());
 	UI_COMMAND(SaveTexture, "Save", "Saves the modified textures for the selected mesh components", EUserInterfaceActionType::Button, FInputChord());
 	UI_COMMAND(Add, "Add", "Adds mesh paint textures to the selected mesh components to enable painting", EUserInterfaceActionType::Button, FInputChord());
@@ -62,22 +61,22 @@ void FMeshPaintEditorModeCommands::RegisterCommands()
 	UI_COMMAND(NextTexture, "Next Texture", "Cycle To next texture", EUserInterfaceActionType::Button, FInputChord(EKeys::Period));
 
 	TArray<TSharedPtr<FUICommandInfo>> VertexColorCommands = {
-		SelectVertex, PaintVertexColor, SwapColor, FillVertex, Propagate, ApplyVertex, SaveVertex, RemoveVertex, Copy, Paste, Import, GetTextureColors, FixVertex
+		SelectVertex, PaintVertexColor, SwapColor, FillVertex, PropagateLODs, PropagateMesh, SaveVertex, RemoveVertex, Copy, Paste, Import, GetTextureColors, FixVertex
 	};
 	Commands.Add(UMeshPaintMode::MeshPaintMode_VertexColor, VertexColorCommands);
 
 	TArray<TSharedPtr<FUICommandInfo>> VertexWeightCommands = {
-		SelectVertex, PaintVertexWeight, FillVertex, Propagate, ApplyVertex, SaveVertex, RemoveVertex, Copy, Paste, Import, FixVertex
+		SelectVertex, PaintVertexWeight, FillVertex, PropagateLODs, PropagateMesh, SaveVertex, RemoveVertex, Copy, Paste, Import, FixVertex
 	};
 	Commands.Add(UMeshPaintMode::MeshPaintMode_VertexWeights, VertexWeightCommands);
 
 	TArray<TSharedPtr<FUICommandInfo>> TextureColorCommands = {
-		SelectTextureColor, PaintTextureColor, SwapColor, FillTexture, ApplyTexture, SaveTexture, Add, RemoveTexture, Copy, Paste, GetVertexColors, FixTexture
+		SelectTextureColor, PaintTextureColor, SwapColor, FillTexture, SaveTexture, Add, RemoveTexture, Copy, Paste, GetVertexColors, FixTexture
 	};
 	Commands.Add(UMeshPaintMode::MeshPaintMode_TextureColor, TextureColorCommands);
 
 	TArray<TSharedPtr<FUICommandInfo>> TextureAssetCommands = {
-		SelectTextureAsset, PaintTextureAsset, SwapColor, FillTexture, ApplyTexture, SaveTexture
+		SelectTextureAsset, PaintTextureAsset, SwapColor, FillTexture, SaveTexture
 	};
 	Commands.Add(UMeshPaintMode::MeshPaintMode_TextureAsset, TextureAssetCommands);
 }
diff --git a/Engine/Plugins/MeshPainting/Source/MeshPaintEditorMode/Public/MeshPaintModeCommands.h b/Engine/Plugins/MeshPainting/Source/MeshPaintEditorMode/Public/MeshPaintModeCommands.h
--- a/Engine/Plugins/MeshPainting/Source/MeshPaintEditorMode/Public/MeshPaintModeCommands.h
+++ b/Engine/Plugins/MeshPainting/Source/MeshPaintEditorMode/Public/MeshPaintModeCommands.h
@@ -42,9 +42,8 @@ public:
 	TSharedPtr<FUICommandInfo> SwapColor;
 	TSharedPtr<FUICommandInfo> FillVertex;
 	TSharedPtr<FUICommandInfo> FillTexture;
-	TSharedPtr<FUICommandInfo> ApplyVertex;
-	TSharedPtr<FUICommandInfo> ApplyTexture;
-	TSharedPtr<FUICommandInfo> Propagate;
+	TSharedPtr<FUICommandInfo> PropagateMesh;
+	TSharedPtr<FUICommandInfo> PropagateLODs;
 	TSharedPtr<FUICommandInfo> SaveVertex;
 	TSharedPtr<FUICommandInfo> SaveTexture;
 	TSharedPtr<FUICommandInfo> Add;
diff --git a/Engine/Plugins/MeshPainting/Source/MeshPaintingToolset/Private/MeshTexturePaintingTool.cpp b/Engine/Plugins/MeshPainting/Source/MeshPaintingToolset/Private/MeshTexturePaintingTool.cpp
--- a/Engine/Plugins/MeshPainting/Source/MeshPaintingToolset/Private/MeshTexturePaintingTool.cpp
+++ b/Engine/Plugins/MeshPainting/Source/MeshPaintingToolset/Private/MeshTexturePaintingTool.cpp
@@ -1,6 +1,7 @@
 // Copyright Epic Games, Inc. All Rights Reserved.
 
 #include "MeshTexturePaintingTool.h"
+
 #include "AssetRegistry/AssetData.h"
 #include "InteractiveToolManager.h"
 #include "Components/MeshComponent.h"
@@ -11,6 +12,7 @@
 #include "Engine/Texture2D.h"
 #include "CanvasItem.h"
 #include "CanvasTypes.h"
+#include "MeshPaintHelpers.h"
 #include "MeshVertexPaintingTool.h"
 #include "ScopedTransaction.h"
 #include "BaseGizmos/BrushStampIndicator.h"
@@ -32,108 +34,6 @@ extern UNREALED_API class UEditorEngine* GEditor;
 
 #define LOCTEXT_NAMESPACE "MeshTextureBrush"
 
-namespace UE::MeshPaintingToolset::Private::TexturePainting
-{
-
-	// Compare two sorted arrays of unique elements
-	template<typename T>
-	void CompareSortedArrayElements(const TArrayView<T>& FirstArray, const TArrayView<T>& SecondArray, TFunctionRef<void(T&)> MissingInFirstArray, TFunctionRef<void(T&)> MissingInSecondArray,  TFunctionRef<void(T&)> PresentInBoth)
-	{
-		T* ElementFromFirstArray = nullptr;
-		T* ElementFromSecondArray = nullptr;
-
-		if (FirstArray.IsEmpty())
-		{
-			for (T& Element : SecondArray)
-			{
-				MissingInFirstArray(Element);
-			}
-			return;
-		}
-		else
-		{
-			ElementFromFirstArray = &FirstArray[0];
-		}
-
-		if (SecondArray.IsEmpty())
-		{
-			for (T& Element : FirstArray)
-			{
-				MissingInSecondArray(Element);
-			}
-			return;
-		}
-		else
-		{
-			ElementFromSecondArray = &SecondArray[0];
-		}
-
-		if (!ElementFromFirstArray || !ElementFromSecondArray)
-		{
-			return;
-		}
-
-		int32 FirstIndex = 0;
-		int32 SecondIndex = 0;
-
-		while (FirstIndex < FirstArray.Num() && SecondIndex < SecondArray.Num())
-		{
-			if (FirstArray.IsValidIndex(FirstIndex))
-			{
-				ElementFromFirstArray = &FirstArray[FirstIndex];
-			}
-
-			if (SecondArray.IsValidIndex(SecondIndex))
-			{
-				ElementFromSecondArray = &SecondArray[SecondIndex];
-			}
-
-			if (*ElementFromFirstArray == *ElementFromSecondArray)
-			{
-				++FirstIndex;
-				++SecondIndex;
-				PresentInBoth(*ElementFromFirstArray);
-			}
-			else if (*ElementFromFirstArray < *ElementFromSecondArray)
-			{
-				MissingInSecondArray(*ElementFromFirstArray);
-				++FirstIndex;
-			}
-			else
-			{
-				MissingInFirstArray(*ElementFromSecondArray);
-				++SecondIndex;
-			}
-		}
-
-		if (FirstArray.Num() < SecondArray.Num())
-		{
-			if (*ElementFromFirstArray != *ElementFromSecondArray)
-			{
-				++SecondIndex;
-			}
-			while (SecondArray.IsValidIndex(SecondIndex))
-			{
-				MissingInFirstArray(SecondArray[SecondIndex]);
-				++SecondIndex;
-			}
-		}
-		else if (SecondArray.Num() < FirstArray.Num())
-		{
-			if (*ElementFromFirstArray != *ElementFromSecondArray)
-			{
-				++FirstIndex;
-			}
-			while (FirstArray.IsValidIndex(FirstIndex))
-			{
-				MissingInSecondArray(FirstArray[FirstIndex]);
-				++FirstIndex;
-			}
-		}
-	}
-}
-
-
 /*
  * ToolBuilder
  */
@@ -172,11 +72,6 @@ void UMeshTexturePaintingTool::Setup()
 {
 	Super::Setup();
 
-	if (UTransBuffer* TransBuffer = GUnrealEd ? Cast<UTransBuffer>(GUnrealEd->Trans) : nullptr)
-	{
-		TransBuffer->OnTransactionStateChanged().AddUObject(this, &UMeshTexturePaintingTool::OnTransactionStateChanged);
-	}
-
 	TextureProperties = Cast<UMeshTexturePaintingToolProperties>(BrushProperties);
 	bResultValid = false;
 	bStampPending = false;
@@ -192,8 +87,12 @@ void UMeshTexturePaintingTool::Setup()
 
 	SelectionMechanic = NewObject<UMeshPaintSelectionMechanic>(this);
 	SelectionMechanic->Setup(this);
-}
 
+	if (UMeshPaintingSubsystem* MeshPaintingSubsystem = GEngine->GetEngineSubsystem<UMeshPaintingSubsystem>())
+	{
+		MeshPaintingSubsystem->Refresh();
+	}
+}
 
 void UMeshTexturePaintingTool::Shutdown(EToolShutdownType ShutdownType)
 {
@@ -210,11 +109,6 @@ void UMeshTexturePaintingTool::Shutdown(EToolShutdownType ShutdownType)
 
 	BrushProperties->SaveProperties(this);
 
-	if (UTransBuffer* TransBuffer = GUnrealEd ? Cast<UTransBuffer>(GUnrealEd->Trans) : nullptr)
-	{
-		TransBuffer->OnTransactionStateChanged().RemoveAll(this);
-	}
-
 	Super::Shutdown(ShutdownType);
 }
 
@@ -259,9 +153,8 @@ void UMeshTexturePaintingTool::Render(IToolsContextRenderAPI* RenderAPI)
 	{
 		BrushStampIndicator->bDrawIndicatorLines = false;
 	}
-		Draw.EndFrame();
+	Draw.EndFrame();
 	UpdateResult();
-
 }
 
 void UMeshTexturePaintingTool::OnTick(float DeltaTime)
@@ -298,7 +191,7 @@ void UMeshTexturePaintingTool::OnTick(float DeltaTime)
 				TSharedPtr<IMeshPaintComponentAdapter> MeshAdapter = MeshPaintingSubsystem->GetAdapterForComponent(SelectedComponent);
 				if (IsMeshAdapterSupported(MeshAdapter))
 				{
-					Textures.Empty();
+					TArray<UTexture const*> Textures;
 					const UTexture2D* TargetTexture2D = GetSelectedPaintTexture(SelectedComponent);
 					if (TargetTexture2D)
 					{
@@ -353,11 +246,9 @@ void UMeshTexturePaintingTool::OnTick(float DeltaTime)
 		{
 			CacheSelectionData();
 			CacheTexturePaintData();
-			bDoRestoreRenTargets = true;
 
-			// We are recaching so we need to clear any of the current overrides that we are caching;
-			// if any overrides are necessary, they will re-register later
 			ClearAllTextureOverrides();
+			SetAllTextureOverrides();
 		}
 	}
 
@@ -379,6 +270,7 @@ void UMeshTexturePaintingTool::OnTick(float DeltaTime)
 		if (TexturePaintingCurrentMeshComponent != nullptr)
 		{
 			FinishPaintingTexture();
+			FinishPainting();
 		}
 
 		bRequestPaintBucketFill = false;
@@ -439,7 +331,6 @@ bool UMeshTexturePaintingTool::Paint(const TArrayView<TPair<FVector, FVector>>&
 	return PaintInternal(Rays, PaintAction, PaintStrength);
 }
 
-
 void UMeshTexturePaintingTool::CacheSelectionData()
 {
 	UMeshPaintingSubsystem* MeshPaintingSubsystem = GEngine->GetEngineSubsystem<UMeshPaintingSubsystem>();
@@ -569,7 +460,6 @@ bool UMeshTexturePaintingTool::PaintInternal(const TArrayView<TPair<FVector, FVe
 		if (bArePainting == false)
 		{
 			bArePainting = true;
-			TimeSinceStartedPainting = 0.0f;
 		}
 
 		// Iterate over the selected meshes under the cursor and paint them!
@@ -583,7 +473,7 @@ bool UMeshTexturePaintingTool::PaintInternal(const TArrayView<TPair<FVector, FVe
 			{
 				const int32 UVChannel = GetSelectedUVChannel(HoveredComponent);
 
-				Textures.Empty();
+				TArray<UTexture const*> Textures;
 				const UTexture2D* TargetTexture2D = GetSelectedPaintTexture(HoveredComponent);
 				if (TargetTexture2D)
 				{
@@ -656,17 +546,13 @@ void UMeshTexturePaintingTool::AddTextureOverrideToComponent(FPaintTexture2DData
 {
 	if (MeshComponent && MeshPaintAdapter)
 	{
-		if (!TextureData.PaintedComponents.Contains(MeshComponent))
+		if (!TextureData.TextureOverrideComponents.Contains(MeshComponent))
 		{
-			TextureData.PaintedComponents.AddUnique(MeshComponent);
+			TextureData.TextureOverrideComponents.AddUnique(MeshComponent);
 
 			MeshPaintAdapter->ApplyOrRemoveTextureOverride(TextureData.PaintingTexture2D, GetTextureForMaterialOverride(TextureData));
-
-			// For the transactions add to the cache of overridden component.
-			PaintComponentsOverride.FindOrAdd(TextureData.PaintingTexture2D).PaintedComponents.Add(MeshComponent);
 		}
 	}
-
 }
 
 void UMeshTexturePaintingTool::UpdateResult()
@@ -676,16 +562,6 @@ void UMeshTexturePaintingTool::UpdateResult()
 	bResultValid = true;
 }
 
-bool UMeshTexturePaintingTool::HasAccept() const
-{
-	return false;
-}
-
-bool UMeshTexturePaintingTool::CanAccept() const
-{
-	return false;
-}
-
 FInputRayHit UMeshTexturePaintingTool::CanBeginClickDragSequence(const FInputDeviceRay& PressPos)
 {
 	FHitResult OutHit;
@@ -707,8 +583,6 @@ FInputRayHit UMeshTexturePaintingTool::CanBeginClickDragSequence(const FInputDev
 	if (MeshPaintingSubsystem && LastBestHitResult.Component != nullptr && MeshPaintingSubsystem->LastPaintedComponent != LastBestHitResult.Component)
 	{
 		MeshPaintingSubsystem->LastPaintedComponent = (UMeshComponent*)LastBestHitResult.Component.Get();
-		//ApplyAllPaintedTextures();
-		//MeshPaintingSubsystem->bNeedsRecache = true;
 	}
 
 	return Super::CanBeginClickDragSequence(PressPos);
@@ -752,8 +626,6 @@ void UMeshTexturePaintingTool::OnUpdateDrag(const FRay& Ray)
 	}
 }
 
-
-
 void UMeshTexturePaintingTool::OnEndDrag(const FRay& Ray)
 {
 	FinishPaintingTexture();
@@ -776,16 +648,10 @@ bool UMeshTexturePaintingTool::HitTest(const FRay& Ray, FHitResult& OutHit)
 
 void UMeshTexturePaintingTool::FinishPainting()
 {
-	if (bArePainting)
-	{
-		bArePainting = false;
-		OnPaintingFinishedDelegate.ExecuteIfBound();
-		PaintingTransaction.Reset();
-	}
+	PaintingTransaction.Reset();
+	bArePainting = false;
 }
 
-
-
 FPaintTexture2DData* UMeshTexturePaintingTool::GetPaintTargetData(const UTexture2D* InTexture)
 {
 	checkf(InTexture != nullptr, TEXT("Invalid Texture ptr"));
@@ -804,7 +670,7 @@ FPaintTexture2DData* UMeshTexturePaintingTool::AddPaintTargetData(UTexture2D* In
 	{
 		// If we didn't find data associated with this texture we create a new entry and return a reference to it.
 		//   Note: This reference is only valid until the next change to any key in the map.
-		TextureData = &PaintTargetData.Add(InTexture, FPaintTexture2DData(InTexture, false));
+		TextureData = &PaintTargetData.Add(InTexture, FPaintTexture2DData(InTexture));
 	}
 	return TextureData;
 }
@@ -838,15 +704,17 @@ void UMeshTexturePaintingTool::GatherTextureTriangles(IMeshPaintComponentAdapter
 	}
 }
 
-
 void UMeshTexturePaintingTool::StartPaintingTexture(UMeshComponent* InMeshComponent, const IMeshPaintComponentAdapter& GeometryInfo)
 {
 	check(InMeshComponent != nullptr);
 	check(TexturePaintingCurrentMeshComponent == nullptr);
 	check(PaintingTexture2D == nullptr);
 
-	PaintingTransaction = MakeUnique<FScopedTransaction>(LOCTEXT("MeshPaintMode_TexturePaint_Transaction", "Texture Paint"));
-	Modify();
+	// Only start new transaction if not in one currently
+	if (!PaintingTransaction.IsValid())
+	{
+		PaintingTransaction = MakeUnique<FScopedTransaction>(LOCTEXT("MeshPaintMode_TexturePaint_Transaction", "Texture Paint"));
+	}
 
 	const auto FeatureLevel = InMeshComponent->GetWorld()->GetFeatureLevel();
 
@@ -882,44 +750,8 @@ void UMeshTexturePaintingTool::StartPaintingTexture(UMeshComponent* InMeshCompon
 			const int32 TextureWidth = Texture2D->Source.GetSizeX();
 			const int32 TextureHeight = Texture2D->Source.GetSizeY();
 
-			Texture2D->BlockOnAnyAsyncBuild();
-
-			if (TextureData == nullptr)
-			{
-				TextureData = AddPaintTargetData(Texture2D);
-			}
 			check(TextureData != nullptr);
 
-			// Create our render target texture
-			if (TextureData->PaintRenderTargetTexture == nullptr ||
-				TextureData->PaintRenderTargetTexture->GetSurfaceWidth() != TextureWidth ||
-				TextureData->PaintRenderTargetTexture->GetSurfaceHeight() != TextureHeight)
-			{
-				TextureData->PaintRenderTargetTexture = nullptr;
-				TextureData->PaintRenderTargetTexture = NewObject<UTextureRenderTarget2D>(GetTransientPackage(), NAME_None, RF_Transient);
-				TextureData->PaintRenderTargetTexture->bNeedsTwoCopies = true;
-				const bool bForceLinearGamma = true;
-				TextureData->PaintRenderTargetTexture->InitCustomFormat(TextureWidth, TextureHeight, PF_A16B16G16R16, bForceLinearGamma);
-				TextureData->PaintRenderTargetTexture->UpdateResourceImmediate();
-
-				TextureData->PaintRenderTargetTextureAdapter = nullptr;
-				if (TextureData->PaintingTexture2D->IsCurrentlyVirtualTextured())
-				{
-					FVirtualTextureBuildSettings VirtualTextureBuildSettings;
-					TextureData->PaintingTexture2D->GetVirtualTextureBuildSettings(VirtualTextureBuildSettings);
-
-					TextureData->PaintRenderTargetTextureAdapter = NewObject<UVirtualTextureAdapter>(GetTransientPackage(), NAME_None, RF_Transient);
-					TextureData->PaintRenderTargetTextureAdapter->Texture = TextureData->PaintRenderTargetTexture;
-					TextureData->PaintRenderTargetTextureAdapter->bUseCompressedFormat = true;
-					TextureData->PaintRenderTargetTextureAdapter->bUseDefaultTileSizes = false;
-					TextureData->PaintRenderTargetTextureAdapter->TileSize = VirtualTextureBuildSettings.TileSize;
-					TextureData->PaintRenderTargetTextureAdapter->TileBorderSize = VirtualTextureBuildSettings.TileBorderSize;
-					TextureData->PaintRenderTargetTextureAdapter->UpdateResource();
-				}
-			}
-			TextureData->PaintRenderTargetTexture->AddressX = Texture2D->AddressX;
-			TextureData->PaintRenderTargetTexture->AddressY = Texture2D->AddressY;
-
 			const int32 BrushTargetTextureWidth = TextureWidth;
 			const int32 BrushTargetTextureHeight = TextureHeight;
 
@@ -976,10 +808,6 @@ void UMeshTexturePaintingTool::StartPaintingTexture(UMeshComponent* InMeshCompon
 				TextureData->bGenerateSeamMask = true;
 			}
 
-			// MeshPaint: Here we override the textures on the mesh with the render target.  The problem is that other meshes in the scene that use
-			// this texture do not get the override. Do we want to extend this to all other selected meshes or maybe even to all meshes in the scene?
-			AddTextureOverrideToComponent(*TextureData, InMeshComponent, &GeometryInfo);
-
 			bStartedPainting = true;
 			UTexture2D* Texture2DPaintBrush = TextureProperties->PaintBrush;
 			if (Texture2DPaintBrush)
@@ -1006,7 +834,6 @@ void UMeshTexturePaintingTool::StartPaintingTexture(UMeshComponent* InMeshCompon
 				TextureData->PaintBrushRenderTargetTexture = nullptr;
 			}
 
-
 			TexturePaintingCurrentMeshComponent = InMeshComponent;
 
 			check(Texture2D != nullptr);
@@ -1023,9 +850,7 @@ void UMeshTexturePaintingTool::StartPaintingTexture(UMeshComponent* InMeshCompon
 
 		check(Texture2D != nullptr);
 		PaintingTexture2D = Texture2D;
-		// Todo investigate if this could be done before the first draw to avoid recreating the scratch texture because of a undo.
-		// OK, now we need to make sure our render target is filled in with data
-		UTexturePaintToolset::SetupInitialRenderTargetData(*TextureData);
+
 		if (TextureProperties->PaintBrush != nullptr && TextureData->PaintBrushRenderTargetTexture != nullptr)
 		{
 			UTexturePaintToolset::SetupInitialRenderTargetData(TextureProperties->PaintBrush, TextureData->PaintBrushRenderTargetTexture);
@@ -1366,25 +1191,19 @@ void UMeshTexturePaintingTool::PaintTexture(FMeshPaintParameters& InParams, int3
 		});
 	}
 
+	// Need to flush the virtual texture adapter since we just updated the painting render target.
 	if (TextureData->PaintRenderTargetTextureAdapter)
 	{
 		TextureData->PaintRenderTargetTextureAdapter->Flush(FBox2f(FVector2f(0, 0), FVector2f(1, 1)));
 	}
-
-	FlushRenderingCommands();
 }
 
 
 void UMeshTexturePaintingTool::FinishPaintingTexture()
 {
-	if (TexturePaintingCurrentMeshComponent != nullptr)
+	if (FPaintTexture2DData* TextureData = GetPaintTargetData(PaintingTexture2D))
 	{
-		check(PaintingTexture2D != nullptr);
-
-		FPaintTexture2DData* TextureData = GetPaintTargetData(PaintingTexture2D);
-		check(TextureData);
-
-		// Set to the texture source art but don't do any compression, compression is saved for the ApplyAllPaintedTextures function.
+		// Apply the texture
 		if (TextureData->bIsPaintingTexture2DModified == true)
 		{
 			const int32 TexWidth = TextureData->PaintRenderTargetTexture->SizeX;
@@ -1392,294 +1211,137 @@ void UMeshTexturePaintingTool::FinishPaintingTexture()
 			TArray< FColor > TexturePixels;
 			TexturePixels.AddUninitialized(TexWidth * TexHeight);
 
+			// Copy the contents of the remote texture to system memory
+
 			FlushRenderingCommands();
 			// NOTE: You are normally not allowed to dereference this pointer on the game thread! Normally you can only pass the pointer around and
 			//  check for NULLness.  We do it in this context, however, and it is only ok because this does not happen every frame and we make sure to flush the
 			//  rendering thread.
 			FTextureRenderTargetResource* RenderTargetResource = TextureData->PaintRenderTargetTexture->GameThread_GetRenderTargetResource();
 			check(RenderTargetResource != nullptr);
-			RenderTargetResource->ReadPixels(TexturePixels);
-
-			{
-				TextureData->ScratchTexture->PreEditChange(nullptr);
-
-				const int32 NumPixels = TexturePixels.Num();
-
-				constexpr int32 NumMips = 1;
-				constexpr int32 NumSlices = 1;
-				// Store source art into the scratch texture
-				TextureData->ScratchTexture->Source.Init(
-						TextureData->PaintingTexture2D->Source.GetSizeX(),
-						TextureData->PaintingTexture2D->Source.GetSizeY(),
-						NumSlices,
-						NumMips,
-						TSF_BGRA8,
-						UE::Serialization::FEditorBulkData::FSharedBufferWithID(MakeSharedBufferFromArray(MoveTemp(TexturePixels)))
-					);
-
-				// Store source art into the scratch texture
-				check(TextureData->ScratchTexture->Source.CalcMipSize(0) == NumPixels * sizeof(FColor));
-
-				// If render target gamma used was 1.0 then disable SRGB for the static texture
-				TextureData->ScratchTexture->SRGB = FMath::Abs(RenderTargetResource->GetDisplayGamma() - 1.0f) >= KINDA_SMALL_NUMBER;
-				TextureData->ScratchTexture->bHasBeenPaintedInEditor = true;
-				TextureData->ScratchTexture->PostEditChange();
-			}
-		}
-
-		PaintingTexture2D = nullptr;
-		TexturePaintingCurrentMeshComponent = nullptr;
-	}
-
-	PaintingTransaction.Reset();
-}
-
-void UMeshTexturePaintingTool::OnTransactionStateChanged(const FTransactionContext& InTransactionContext, const ETransactionStateEventType InTransactionState)
-{
-	if (InTransactionState == ETransactionStateEventType::UndoRedoFinalized)
-	{
-		if (GUnrealEd)
-		{
-			if (UTransactor* Trans = GUnrealEd->Trans)
-			{
-				int32 CurrentTransactionIndex = Trans->FindTransactionIndex(InTransactionContext.TransactionId);
-				if (const FTransaction* Transaction = Trans->GetTransaction(CurrentTransactionIndex))
-				{
-					// Update the texture override of the components
-					using namespace UE::MeshPaintingToolset::Private::TexturePainting;
-
-					UMeshPaintingSubsystem* MeshPaintingSubsystem = GEngine->GetEngineSubsystem<UMeshPaintingSubsystem>();
-					check(MeshPaintingSubsystem);
-					{
-						auto TextureOverrideIsMissing = [this, MeshPaintingSubsystem](TObjectPtr<UTexture2D>& Texture)
-						{
-							if (FPaintTexture2DData* TextureData = GetPaintTargetData(Texture))
-							{
-								FPaintComponentOverride& PaintComponentOverride = PaintComponentsOverride.FindOrAdd(TextureData->PaintingTexture2D);
-								PaintComponentOverride.PaintedComponents.Reserve(TextureData->PaintedComponents.Num());
-
-								for (UMeshComponent* Component : TextureData->PaintedComponents)
-								{
-									if (TSharedPtr<IMeshPaintComponentAdapter> PaintAdapterResource = MeshPaintingSubsystem->GetAdapterForComponent(Component))
-									{
-										PaintAdapterResource->ApplyOrRemoveTextureOverride(TextureData->PaintingTexture2D, GetTextureForMaterialOverride(*TextureData));
-										PaintComponentOverride.PaintedComponents.Add(Component);
-									}
-								}
-							}
-						};
-
-						auto TextureOverrideIsNotNeeded = [this, MeshPaintingSubsystem](TObjectPtr<UTexture2D>& Texture)
-						{
-							FPaintComponentOverride ComponentOverride;
-							PaintComponentsOverride.RemoveAndCopyValue(Texture, ComponentOverride);
-
-							for (UMeshComponent* Component : ComponentOverride.PaintedComponents)
-							{
-								if (TSharedPtr<IMeshPaintComponentAdapter> PaintAdapterResource = MeshPaintingSubsystem->GetAdapterForComponent(Component))
-								{
-									PaintAdapterResource->ApplyOrRemoveTextureOverride(Texture,  nullptr);
-								}
-							}
-						};
-
-
-						auto UpdateOverridedComponents = [this, MeshPaintingSubsystem](TObjectPtr<UTexture2D>& Texture)
-						{
-							if (FPaintTexture2DData* TextureData = GetPaintTargetData(Texture))
-							{
-								FPaintComponentOverride& ComponentOverride = PaintComponentsOverride.FindChecked(Texture);
-
-								auto SortComponents = [](const UMeshComponent& First, const UMeshComponent& Second) { return First < Second; };
-
-								TextureData->PaintedComponents.Sort(SortComponents);
-								ComponentOverride.PaintedComponents.Sort(SortComponents);
-
-								TArray<UMeshComponent*> MeshComponentsToRemove;
-								TArray<UMeshComponent*> MeshComponentsToAdd;
-
-								auto ComponentOverrideNotNeeded = [this, MeshPaintingSubsystem, &ComponentOverride, Texture, &MeshComponentsToRemove](TObjectPtr<UMeshComponent>& MeshComponent)
-								{
-									if (TSharedPtr<IMeshPaintComponentAdapter> PaintAdapterResource = MeshPaintingSubsystem->GetAdapterForComponent(MeshComponent))
-									{
-										PaintAdapterResource->ApplyOrRemoveTextureOverride(Texture, nullptr);
-									}
 
-									MeshComponentsToRemove.Add(MeshComponent);
-								};
+			FReadSurfaceDataFlags Flags;
+			Flags.SetLinearToGamma(PaintingTexture2D->SRGB);
+			RenderTargetResource->ReadPixels(TexturePixels, Flags);
 
-								UTexture* OverrideTexture = GetTextureForMaterialOverride(*TextureData);
-								auto ComponentOverrideMissing = [this, MeshPaintingSubsystem, &ComponentOverride, Texture, &MeshComponentsToAdd, OverrideTexture](TObjectPtr<UMeshComponent>& MeshComponent)
-								{
-									if (TSharedPtr<IMeshPaintComponentAdapter> PaintAdapterResource = MeshPaintingSubsystem->GetAdapterForComponent(MeshComponent))
-									{
-										PaintAdapterResource->ApplyOrRemoveTextureOverride(Texture, OverrideTexture);
-									}
+			// For undo
+			TextureData->PaintingTexture2D->SetFlags(RF_Transactional);
+			TextureData->PaintingTexture2D->PreEditChange(nullptr);
 
-									MeshComponentsToAdd.Add(MeshComponent);
-								};
+			// Store source art
+			// @@ use FMipLock.Image instead
+			FColor* Colors = (FColor*)TextureData->PaintingTexture2D->Source.LockMip(0);
+			check(TextureData->PaintingTexture2D->Source.CalcMipSize(0) == TexturePixels.Num() * sizeof(FColor));
+			FMemory::Memcpy(Colors, TexturePixels.GetData(), TexturePixels.Num() * sizeof(FColor));
+			TextureData->PaintingTexture2D->Source.UnlockMip(0);
 
-								auto DoNothing = [](TObjectPtr<UMeshComponent>&) {};
+			TextureData->PaintingTexture2D->bHasBeenPaintedInEditor = true;
 
-								CompareSortedArrayElements<TObjectPtr<UMeshComponent>>(TextureData->PaintedComponents, ComponentOverride.PaintedComponents, ComponentOverrideNotNeeded, ComponentOverrideMissing, DoNothing);
-							}
-						};
-					
-						// Unfortunately TObjectPtrs don't behave like a RawPtrs when sorting an array 
-						auto SortPredicate = [](const UTexture2D& First, const UTexture2D& Second) {return (&First) < (&Second);};
-
-						TArray<TObjectPtr<UTexture2D>> CurrentTextures;
-						PaintTargetData.GenerateKeyArray(CurrentTextures);
-						CurrentTextures.Sort(SortPredicate);
-
-						TArray<TObjectPtr<UTexture2D>> OverrideTextures;
-						PaintComponentsOverride.GenerateKeyArray(OverrideTextures);
-						PaintComponentsOverride.KeySort(SortPredicate);
-
-						CompareSortedArrayElements<TObjectPtr<UTexture2D>>(CurrentTextures, OverrideTextures, TextureOverrideIsNotNeeded, TextureOverrideIsMissing, UpdateOverridedComponents);
-					}
+			// Update the texture (generate mips, compress if needed)
+			TextureData->PaintingTexture2D->PostEditChange();
 
-					// Update the render targets
-					for (TPair<TObjectPtr<UTexture2D>, FPaintTexture2DData>& TargetData : PaintTargetData)
-					{
-						FPaintTexture2DData& PaintTexture2DData = TargetData.Value;
+			TextureData->bIsPaintingTexture2DModified = false;
 
-						if (Transaction->ContainsObject(PaintTexture2DData.ScratchTexture))
-						{
-							UTexturePaintToolset::UpdateRenderTargetData(PaintTexture2DData);
-						}
-					}
-				}
-			}
+			OnPaintingFinishedDelegate.ExecuteIfBound(TexturePaintingCurrentMeshComponent);
 		}
 	}
-}
 
-void UMeshTextureAssetPaintingTool::CycleTextures(int32 Direction)
-{
-	if (!PaintableTextures.Num())
-	{
-		return;
-	}
-	TObjectPtr<UTexture2D>& SelectedTexture = AssetProperties->PaintTexture;
-	const int32 TextureIndex = (SelectedTexture != nullptr) ? PaintableTextures.IndexOfByKey(SelectedTexture) : 0;
-	if (TextureIndex != INDEX_NONE)
-	{
-		int32 NewTextureIndex = TextureIndex + Direction;
-		if (NewTextureIndex < 0)
-		{
-			NewTextureIndex += PaintableTextures.Num();
-		}
-		NewTextureIndex %= PaintableTextures.Num();
-
-		if (PaintableTextures.IsValidIndex(NewTextureIndex))
-		{
-			SelectedTexture = (UTexture2D*)PaintableTextures[NewTextureIndex].Texture;
-		}
-	}
+	PaintingTexture2D = nullptr;
+	TexturePaintingCurrentMeshComponent = nullptr;
 }
 
-
-void UMeshTexturePaintingTool::ApplyAllPaintedTextures()
+void UMeshTexturePaintingTool::ClearAllTextureOverrides()
 {
-	if (PaintTargetData.Num() > 0)
+	if (UMeshPaintingSubsystem* MeshPaintingSubsystem = GEngine->GetEngineSubsystem<UMeshPaintingSubsystem>())
 	{
-		check(PaintingTexture2D == nullptr);
-
-		Modify();
-
-		int32 CurStep = 1;
-		int32 TotalSteps = GetNumberOfPendingPaintChanges();
-
+		/** Remove all texture overrides which are currently stored and active */
 		for (decltype(PaintTargetData)::TIterator It(PaintTargetData); It; ++It)
 		{
 			FPaintTexture2DData* TextureData = &It.Value();
 
-			// Apply the texture
-			if (TextureData->bIsPaintingTexture2DModified == true)
+			for (UMeshComponent* MeshComponent : TextureData->TextureOverrideComponents)
 			{
-				const int32 TexWidth = TextureData->PaintRenderTargetTexture->SizeX;
-				const int32 TexHeight = TextureData->PaintRenderTargetTexture->SizeY;
-				TArray< FColor > TexturePixels;
-				TexturePixels.AddUninitialized(TexWidth * TexHeight);
-
-				// Copy the contents of the remote texture to system memory
-				// NOTE: OutRawImageData must be a preallocated buffer!
-
-				FlushRenderingCommands();
-				// NOTE: You are normally not allowed to dereference this pointer on the game thread! Normally you can only pass the pointer around and
-				//  check for NULLness.  We do it in this context, however, and it is only ok because this does not happen every frame and we make sure to flush the
-				//  rendering thread.
-				FTextureRenderTargetResource* RenderTargetResource = TextureData->PaintRenderTargetTexture->GameThread_GetRenderTargetResource();
-				check(RenderTargetResource != nullptr);
-				RenderTargetResource->ReadPixels(TexturePixels);
-
+				if (TSharedPtr<IMeshPaintComponentAdapter> PaintAdapter = MeshPaintingSubsystem->GetAdapterForComponent(MeshComponent))
 				{
-					// For undo
-					TextureData->PaintingTexture2D->SetFlags(RF_Transactional);
-					TextureData->PaintingTexture2D->PreEditChange(nullptr);
-
-					// Store source art
-					// @@ use FMipLock.Image instead
-					FColor* Colors = (FColor*)TextureData->PaintingTexture2D->Source.LockMip(0);
-					check(TextureData->PaintingTexture2D->Source.CalcMipSize(0) == TexturePixels.Num() * sizeof(FColor));
-					FMemory::Memcpy(Colors, TexturePixels.GetData(), TexturePixels.Num() * sizeof(FColor));
-					TextureData->PaintingTexture2D->Source.UnlockMip(0);
-
-					TextureData->PaintingTexture2D->bHasBeenPaintedInEditor = true;
-
-					// Update the texture (generate mips, compress if needed)
-					TextureData->PaintingTexture2D->PostEditChange();
-
-					TextureData->bIsPaintingTexture2DModified = false;
-
+					PaintAdapter->ApplyOrRemoveTextureOverride(TextureData->PaintingTexture2D, nullptr);
 				}
 			}
+	
+			TextureData->TextureOverrideComponents.Empty();
 		}
-
-		ClearAllTextureOverrides();
 	}
 }
 
-void UMeshTexturePaintingTool::ClearAllTextureOverrides()
+void UMeshTexturePaintingTool::SetAllTextureOverrides()
 {
-	if (UMeshPaintingSubsystem* MeshPaintingSubsystem = GEngine->GetEngineSubsystem<UMeshPaintingSubsystem>())
+	for (FPaintableTexture const& PaintableTexture : PaintableTextures)
 	{
-		check(GEditor && GEditor->GetEditorWorldContext().World());
-		const auto FeatureLevel = GEditor->GetEditorWorldContext().World()->GetFeatureLevel();
-		/** Remove all texture overrides which are currently stored and active */
-		for (decltype(PaintTargetData)::TIterator It(PaintTargetData); It; ++It)
+		if (UTexture2D* Texture2D = Cast<UTexture2D>(PaintableTexture.Texture))
 		{
-			FPaintTexture2DData* TextureData = &It.Value();
+			Texture2D->BlockOnAnyAsyncBuild();
 
-			for (UMeshComponent* MeshComponent : TextureData->PaintedComponents)
+			// Create the render target to paint on.
+			FPaintTexture2DData* TextureData = AddPaintTargetData(Texture2D);
+
+			const int32 TextureWidth = Texture2D->Source.GetSizeX();
+			const int32 TextureHeight = Texture2D->Source.GetSizeY();
+
+			if (TextureData->PaintRenderTargetTexture == nullptr ||
+				TextureData->PaintRenderTargetTexture->GetSurfaceWidth() != TextureWidth ||
+				TextureData->PaintRenderTargetTexture->GetSurfaceHeight() != TextureHeight)
 			{
-				if (TSharedPtr<IMeshPaintComponentAdapter> PaintAdapter = MeshPaintingSubsystem->GetAdapterForComponent(MeshComponent))
+				TextureData->PaintRenderTargetTexture = NewObject<UTextureRenderTarget2D>(GetTransientPackage(), NAME_None, RF_Transient);
+				TextureData->PaintRenderTargetTexture->bNeedsTwoCopies = true;
+				const bool bForceLinearGamma = true;
+				TextureData->PaintRenderTargetTexture->InitCustomFormat(TextureWidth, TextureHeight, PF_A16B16G16R16, bForceLinearGamma);
+				TextureData->PaintRenderTargetTexture->UpdateResourceImmediate();
+
+				TextureData->PaintRenderTargetTextureAdapter = nullptr; 
+				if (TextureData->PaintingTexture2D->IsCurrentlyVirtualTextured())
 				{
-					PaintAdapter->ApplyOrRemoveTextureOverride(TextureData->PaintingTexture2D, nullptr);
+					// Virtual textures can't just swap in a render target in their material, so we use a virtual texture adapter.
+					FVirtualTextureBuildSettings VirtualTextureBuildSettings;
+					TextureData->PaintingTexture2D->GetVirtualTextureBuildSettings(VirtualTextureBuildSettings);
+
+					TextureData->PaintRenderTargetTextureAdapter = NewObject<UVirtualTextureAdapter>(GetTransientPackage(), NAME_None, RF_Transient);
+					TextureData->PaintRenderTargetTextureAdapter->Texture = TextureData->PaintRenderTargetTexture;
+					TextureData->PaintRenderTargetTextureAdapter->bUseCompressedFormat = true;
+					TextureData->PaintRenderTargetTextureAdapter->bUseDefaultTileSizes = false;
+					TextureData->PaintRenderTargetTextureAdapter->TileSize = VirtualTextureBuildSettings.TileSize;
+					TextureData->PaintRenderTargetTextureAdapter->TileBorderSize = VirtualTextureBuildSettings.TileBorderSize;
+					TextureData->PaintRenderTargetTextureAdapter->UpdateResource();
 				}
 			}
-	
-			TextureData->PaintedComponents.Empty();
-		}
 
-		PaintComponentsOverride.Empty();
-	}
-}
+			TextureData->PaintRenderTargetTexture->AddressX = Texture2D->AddressX;
+			TextureData->PaintRenderTargetTexture->AddressY = Texture2D->AddressY;
 
-int32 UMeshTexturePaintingTool::GetNumberOfPendingPaintChanges() const
-{
-	int32 Result = 0;
-	for (decltype(PaintTargetData)::TConstIterator It(PaintTargetData); It; ++It)
-	{
-		const FPaintTexture2DData* TextureData = &It.Value();
+			// Initialize the render target with the texture contents.
+			UTexturePaintToolset::SetupInitialRenderTargetData(TextureData->PaintingTexture2D, TextureData->PaintRenderTargetTexture);
 
-		if (TextureData->bIsPaintingTexture2DModified == true)
-		{
-			Result++;
+			// Need to flush the virtual texture adapter since we just updated the painting render target.
+			if (TextureData->PaintRenderTargetTextureAdapter)
+			{
+				TextureData->PaintRenderTargetTextureAdapter->Flush(FBox2f(FVector2f(0, 0), FVector2f(1, 1)));
+			}
+
+			// Apply the render target override only to the components that we are painting with this texture.
+			UMeshPaintingSubsystem* MeshPaintingSubsystem = GEngine->GetEngineSubsystem<UMeshPaintingSubsystem>();
+			TArray<UMeshComponent*> SelectedMeshComponents = MeshPaintingSubsystem->GetSelectedMeshComponents();
+			for (int32 Index = 0; Index < SelectedMeshComponents.Num(); ++Index)
+			{
+				UMeshComponent* MeshComponent = SelectedMeshComponents[Index];
+				if (CanPaintTextureToComponent(PaintableTexture.Texture, MeshComponent))
+				{
+					TSharedPtr<IMeshPaintComponentAdapter> MeshAdapter = MeshPaintingSubsystem->GetAdapterForComponent(MeshComponent);
+					if (IsMeshAdapterSupported(MeshAdapter))
+					{
+						AddTextureOverrideToComponent(*TextureData, MeshComponent, MeshAdapter.Get());
+					}
+				}
+			}
 		}
 	}
-	return Result;
 }
 
 void UMeshTexturePaintingTool::FloodCurrentPaintTexture()
@@ -1748,28 +1410,9 @@ void UMeshTextureColorPaintingTool::CacheTexturePaintData()
 	}
 }
 
-void UMeshTextureColorPaintingTool::GetPaintedComponentsForPropagateVertexColor(TArray<UMeshComponent*>& OutComponents) const
+bool UMeshTextureColorPaintingTool::CanPaintTextureToComponent(UTexture* InTexture, UMeshComponent const* InMeshComponent) const
 {
-	if (!ColorProperties->bPropagateToVertexColor)
-	{
-		return;
-	}
-	
-	for (decltype(PaintTargetData)::TConstIterator It(PaintTargetData); It; ++It)
-	{
-		const FPaintTexture2DData* TextureData = &It.Value();
-
-		if (TextureData->bIsPaintingTexture2DModified == true)
-		{
-			for (UMeshComponent* Component : TextureData->PaintedComponents)
-			{
-				if (Component->GetMeshPaintTexture() == TextureData->PaintingTexture2D)
-				{
-					OutComponents.Add(Component);
-				}
-			}
-		}
-	}
+	return InMeshComponent->GetMeshPaintTexture() == InTexture;
 }
 
 
@@ -1821,18 +1464,25 @@ bool UMeshTextureAssetPaintingTool::ShouldFilterTextureAsset(const FAssetData& A
 	return !(PaintableTextures.ContainsByPredicate([=](const FPaintableTexture& Texture) { return Texture.Texture->GetFullName() == AssetTexture->GetFullName(); }));
 }
 
-void UMeshTextureAssetPaintingTool::PaintTextureChanged(const FAssetData& AssetData)
+void UMeshTextureAssetPaintingTool::OnPropertyModified(UObject* PropertySet, FProperty* Property)
 {
-	if (UTexture2D* Texture = Cast<UTexture2D>(AssetData.GetAsset()))
+	Super::OnPropertyModified(PropertySet, Property);
+
+	if (Property && (Property->GetFName() == GET_MEMBER_NAME_CHECKED(UMeshTextureAssetPaintingToolProperties, PaintTexture)))
 	{
+		// Find the selected texture and apply it's UV channel.
 		for (FPaintableTexture& PaintableTexture : PaintableTextures)
 		{
-			if (PaintableTexture.Texture == Texture)
+			if (PaintableTexture.Texture == AssetProperties->PaintTexture)
 			{
 				AssetProperties->UVChannel = PaintableTexture.UVChannelIndex;
 				break;
 			}
 		}
+
+		// Need to recreate the render target overrides with the newly selected texture.
+		ClearAllTextureOverrides();
+		SetAllTextureOverrides();
 	}
 }
 
@@ -1882,6 +1532,11 @@ void UMeshTextureAssetPaintingTool::CacheTexturePaintData()
 	}
 }
 
+bool UMeshTextureAssetPaintingTool::CanPaintTextureToComponent(UTexture* InTexture, UMeshComponent const* InMeshComponent) const
+{
+	return InTexture == AssetProperties->PaintTexture;
+}
+
 UTexture* UMeshTextureAssetPaintingTool::GetSelectedPaintTextureWithOverride() const
 {
 	UTexture* SelectedTexture = AssetProperties->PaintTexture;
@@ -1896,4 +1551,28 @@ UTexture* UMeshTextureAssetPaintingTool::GetSelectedPaintTextureWithOverride() c
 	return SelectedTexture;
 }
 
+void UMeshTextureAssetPaintingTool::CycleTextures(int32 Direction)
+{
+	if (!PaintableTextures.Num())
+	{
+		return;
+	}
+	TObjectPtr<UTexture2D>& SelectedTexture = AssetProperties->PaintTexture;
+	const int32 TextureIndex = (SelectedTexture != nullptr) ? PaintableTextures.IndexOfByKey(SelectedTexture) : 0;
+	if (TextureIndex != INDEX_NONE)
+	{
+		int32 NewTextureIndex = TextureIndex + Direction;
+		if (NewTextureIndex < 0)
+		{
+			NewTextureIndex += PaintableTextures.Num();
+		}
+		NewTextureIndex %= PaintableTextures.Num();
+
+		if (PaintableTextures.IsValidIndex(NewTextureIndex))
+		{
+			SelectedTexture = (UTexture2D*)PaintableTextures[NewTextureIndex].Texture;
+		}
+	}
+}
+
 #undef LOCTEXT_NAMESPACE
diff --git a/Engine/Plugins/MeshPainting/Source/MeshPaintingToolset/Private/MeshVertexPaintingTool.cpp b/Engine/Plugins/MeshPainting/Source/MeshPaintingToolset/Private/MeshVertexPaintingTool.cpp
--- a/Engine/Plugins/MeshPainting/Source/MeshPaintingToolset/Private/MeshVertexPaintingTool.cpp
+++ b/Engine/Plugins/MeshPainting/Source/MeshPaintingToolset/Private/MeshVertexPaintingTool.cpp
@@ -162,7 +162,6 @@ void UMeshVertexPaintingTool::OnTick(float DeltaTime)
 		if (MeshPaintingSubsystem->bNeedsRecache)
 		{
 			CacheSelectionData();
-			bDoRestoreRenTargets = true;
 		}
 	}
 
@@ -337,7 +336,6 @@ bool UMeshVertexPaintingTool::PaintInternal(const TArrayView<TPair<FVector, FVec
 				// Vertex painting is an ongoing transaction, while texture painting is handled separately later in a single transaction
 				GetToolManager()->BeginUndoTransaction(LOCTEXT("MeshPaintMode_VertexPaint_TransactionPaintStroke", "Vertex Paint"));
 				bArePainting = true;
-				TimeSinceStartedPainting = 0.0f;
 		}
 
 		// Iterate over the selected meshes under the cursor and paint them!
diff --git a/Engine/Plugins/MeshPainting/Source/MeshPaintingToolset/Private/TexturePaintToolset.cpp b/Engine/Plugins/MeshPainting/Source/MeshPaintingToolset/Private/TexturePaintToolset.cpp
--- a/Engine/Plugins/MeshPainting/Source/MeshPaintingToolset/Private/TexturePaintToolset.cpp
+++ b/Engine/Plugins/MeshPainting/Source/MeshPaintingToolset/Private/TexturePaintToolset.cpp
@@ -414,72 +414,6 @@ void UTexturePaintToolset::SetupInitialRenderTargetData(UTexture2D* InTextureSou
 	}
 }
 
-void UTexturePaintToolset::SetupInitialRenderTargetData(FPaintTexture2DData& PaintTextureData)
-{
-	check(PaintTextureData.PaintingTexture2D != nullptr);
-	check(PaintTextureData.PaintRenderTargetTexture != nullptr);
-
-	if (PaintTextureData.PaintingTexture2D->Source.IsValid())
-	{
-		// Great, we have source data!  We'll use that as our image source.
-
-		// Create a texture in memory from the source art
-		{
-			if (!PaintTextureData.ScratchTexture)
-			{
-				PaintTextureData.ScratchTexture = CreateScratchUncompressedTexture(PaintTextureData.PaintingTexture2D);
-
-				check(PaintTextureData.ScratchTexture != nullptr);
-
-				// Copy the texture to the render target using the GPU
-				CopyTextureToRenderTargetTexture(PaintTextureData.ScratchTexture, PaintTextureData.PaintRenderTargetTexture, GEditor->GetEditorWorldContext().World()->GetFeatureLevel());
-			}
-			
-			// No need to update the render target if the scratch texture exist the paint operations and the undo/redo keep the render target up to date
-		}
-	}
-	else
-	{
-		// Just copy (render) the texture in GPU memory to our render target.  Hopefully it's not
-		// compressed already!
-		check(PaintTextureData.PaintingTexture2D->IsFullyStreamedIn());
-		CopyTextureToRenderTargetTexture(PaintTextureData.PaintingTexture2D, PaintTextureData.PaintRenderTargetTexture, GEditor->GetEditorWorldContext().World()->GetFeatureLevel());
-	}
-}
-
-void UTexturePaintToolset::UpdateRenderTargetData(FPaintTexture2DData& PaintTextureData)
-{
-	check(PaintTextureData.ScratchTexture);
-	check(PaintTextureData.PaintRenderTargetTexture);
-
-
-	TArray64<uint8> RawData;
-	PaintTextureData.ScratchTexture->Source.GetMipData(RawData, 0);
-
-	int32 Width = PaintTextureData.ScratchTexture->Source.GetSizeX();
-	int32 Height = PaintTextureData.ScratchTexture->Source.GetSizeY();
-
-	// Fill in the base mip for the texture we created
-	uint8* MipData = (uint8*)PaintTextureData.ScratchTexture->GetPlatformData()->Mips[0].BulkData.Lock(LOCK_READ_WRITE);
-	for (int32 y = 0; y < Height; y++)
-	{
-		uint8* DestPtr = &MipData[(Height - 1 - y) * Width * sizeof(FColor)];
-		const FColor* SrcPtr = &((FColor*)(RawData.GetData()))[(Height - 1 - y) * Width];
-		for (int32 x = 0; x < Width; x++)
-		{
-			*DestPtr++ = SrcPtr->B;
-			*DestPtr++ = SrcPtr->G;
-			*DestPtr++ = SrcPtr->R;
-			*DestPtr++ = SrcPtr->A;
-			SrcPtr++;
-		}
-	}
-	PaintTextureData.ScratchTexture->GetPlatformData()->Mips[0].BulkData.Unlock();
-	PaintTextureData.ScratchTexture->UpdateResource();
-
-	CopyTextureToRenderTargetTexture(PaintTextureData.ScratchTexture, PaintTextureData.PaintRenderTargetTexture, GEditor->GetEditorWorldContext().World()->GetFeatureLevel());
-}
-
 void UTexturePaintToolset::FindMaterialIndicesUsingTexture(const UTexture* Texture, const UMeshComponent* MeshComponent, TArray<int32>& OutIndices)
 {
 	checkf(Texture && MeshComponent, TEXT("Invalid Texture of MeshComponent"));
@@ -559,13 +493,6 @@ void UTexturePaintToolset::RetrieveMeshSectionsForMaterialIndices(const UMeshCom
 	}
 }
 
-bool UTexturePaintToolset::DoesMeshComponentUseTexture(UMeshComponent* MeshComponent, UTexture* Texture)
-{
-	TArray<UTexture*> UsedTextures;
-	MeshComponent->GetUsedTextures(UsedTextures, EMaterialQualityLevel::High);
-	return UsedTextures.Contains(Texture);
-}
-
 void UTexturePaintToolset::RetrieveTexturesForComponent(const UMeshComponent* Component, IMeshPaintComponentAdapter* Adapter, int32& OutDefaultIndex, TArray<FPaintableTexture>& OutTextures)
 {
 	OutDefaultIndex = INDEX_NONE;
diff --git a/Engine/Plugins/MeshPainting/Source/MeshPaintingToolset/Public/MeshPaintingToolsetTypes.h b/Engine/Plugins/MeshPainting/Source/MeshPaintingToolset/Public/MeshPaintingToolsetTypes.h
--- a/Engine/Plugins/MeshPainting/Source/MeshPaintingToolset/Public/MeshPaintingToolsetTypes.h
+++ b/Engine/Plugins/MeshPainting/Source/MeshPaintingToolset/Public/MeshPaintingToolsetTypes.h
@@ -137,29 +137,22 @@ struct FPaintTexture2DData
 	UPROPERTY(Transient)
 	TObjectPtr<UTexture2D> PaintingTexture2D;
 
+	/** Render target texture for painting */
 	UPROPERTY(Transient)
-	bool bIsPaintingTexture2DModified = false;
+	TObjectPtr<UTextureRenderTarget2D> PaintRenderTargetTexture;
 
-	/** Texture used to store the paint modifications for the transactions */
+	/** Optional virtual texture adapter that we can use to visualize PaintRenderTargetTexture in materials that sample virtual textures */
 	UPROPERTY(Transient)
-	TObjectPtr<UTexture2D> ScratchTexture;
+	TObjectPtr<class UVirtualTextureAdapter> PaintRenderTargetTextureAdapter;
 
-	/** Render target texture for painting */
+	/** Array of components that have the PaintRenderTargetTexture set as a texture override */
 	UPROPERTY(Transient)
-	TObjectPtr<UTextureRenderTarget2D> PaintRenderTargetTexture;
+	TArray<TObjectPtr<UMeshComponent>> TextureOverrideComponents;
 
-	/** List of component we are painting on */
-	UPROPERTY(Transient)
-	TArray<TObjectPtr<UMeshComponent>> PaintedComponents;
-	
 	/** Optional render target texture used as an input while painting that contains a clone of the texture painting brush */
 	UPROPERTY(Transient)
 	TObjectPtr<UTextureRenderTarget2D> PaintBrushRenderTargetTexture;
 
-	/** Optional virtual texture adapter that we can use to visualize the render target in materials that use mesh paint virtual textures */
-	UPROPERTY(Transient)
-	TObjectPtr<class UVirtualTextureAdapter> PaintRenderTargetTextureAdapter;
-
 	/** Temporary render target used to draw incremental paint to */
 	UPROPERTY(Transient)
 	TObjectPtr<UTextureRenderTarget2D> BrushRenderTargetTexture;
@@ -168,19 +161,22 @@ struct FPaintTexture2DData
 	UPROPERTY(Transient)
 	TObjectPtr<UTextureRenderTarget2D> BrushMaskRenderTargetTexture;
 
+	/** True if we need to generate a texture seam mask used for texture dilation */
+	UPROPERTY(Transient)
+	bool bGenerateSeamMask = false;
+
 	/** Optional render target used to store generated mask for texture seams. We create this by projecting object triangles into texture space using the selected UV channel. */
 	UPROPERTY(Transient)
 	TObjectPtr<UTextureRenderTarget2D> SeamMaskRenderTargetTexture;
 
-	/** True if we need to generate a texture seam mask used for texture dilation */
+	/** True if we have some painting applied to the PaintRenderTargetTexture. */
 	UPROPERTY(Transient)
-	bool bGenerateSeamMask = false;
+	bool bIsPaintingTexture2DModified = false;
 
 	FPaintTexture2DData() = default;
 
-	FPaintTexture2DData(UTexture2D* InPaintingTexture2D, bool InbIsPaintingTexture2DModified = false)
+	FPaintTexture2DData(UTexture2D* InPaintingTexture2D)
 		: PaintingTexture2D(InPaintingTexture2D)
-		, bIsPaintingTexture2DModified(InbIsPaintingTexture2DModified)
 	{
 	}
 };
diff --git a/Engine/Plugins/MeshPainting/Source/MeshPaintingToolset/Public/MeshTexturePaintingTool.h b/Engine/Plugins/MeshPainting/Source/MeshPaintingToolset/Public/MeshTexturePaintingTool.h
--- a/Engine/Plugins/MeshPainting/Source/MeshPaintingToolset/Public/MeshTexturePaintingTool.h
+++ b/Engine/Plugins/MeshPainting/Source/MeshPaintingToolset/Public/MeshTexturePaintingTool.h
@@ -4,12 +4,8 @@
 
 #include "BaseTools/BaseBrushTool.h"
 #include "BaseMeshPaintingToolProperties.h"
-#include "MeshPaintHelpers.h"
-#include "MeshPaintInteractions.h"
 #include "MeshPaintingToolsetTypes.h"
-#include "MeshVertexPaintingTool.h"
-#include "Misc/ITransaction.h"
-
+#include "MeshPaintInteractions.h"
 #include "MeshTexturePaintingTool.generated.h"
 
 enum class EMeshPaintModeAction : uint8;
@@ -18,12 +14,12 @@ class FScopedTransaction;
 class IMeshPaintComponentAdapter;
 class UMeshToolManager;
 class UTexture2D;
-class UTextureRenderTarget2D;
+struct FPaintRayResults;
 struct FPaintTexture2DData;
 struct FTexturePaintMeshSectionInfo;
-struct FTextureTargetListInfo;
 struct FToolBuilderState;
 
+
 /**
  * Builder for the texture color mesh paint tool.
  */
@@ -146,8 +142,8 @@ public:
 	virtual void Render(IToolsContextRenderAPI* RenderAPI) override;
 	virtual void OnTick(float DeltaTime) override;
 	virtual bool HasCancel() const override { return false; }
-	virtual bool HasAccept() const override;
-	virtual bool CanAccept() const override;
+	virtual bool HasAccept() const override { return false; }
+	virtual bool CanAccept() const override { return false; }
 	virtual FInputRayHit CanBeginClickDragSequence(const FInputDeviceRay& PressPos) override;
 	virtual void OnUpdateModifierState(int ModifierID, bool bIsOn) override;
 	virtual void OnBeginDrag(const FRay& Ray) override;
@@ -156,109 +152,82 @@ public:
 	virtual	bool HitTest(const FRay& Ray, FHitResult& OutHit) override;
 	virtual void OnPropertyModified(UObject* PropertySet, FProperty* Property) override;
 	virtual double EstimateMaximumTargetDimension() override;
-	virtual bool IsPainting() const
-	{
-		return bArePainting;
-	}
-
-	FSimpleDelegate& OnPaintingFinished()
-	{
-		return OnPaintingFinishedDelegate;
-	}
+	virtual bool IsPainting() const { return bArePainting; }
 
-	void ApplyAllPaintedTextures();
-	void ClearAllTextureOverrides();
-	
-	/** Returns the number of texture that require a commit. */
-	int32 GetNumberOfPendingPaintChanges() const;
+	DECLARE_DELEGATE_OneParam(FOnPaintingFinishedDelegate, UMeshComponent*);
+	FOnPaintingFinishedDelegate& OnPaintingFinished() { return OnPaintingFinishedDelegate; }
 
 	void FloodCurrentPaintTexture();
 	
 	virtual void GetModifiedTexturesToSave(TArray<UObject*>& OutTexturesToSave) const {}
 
 protected:
-	virtual void SetAdditionalPaintParameters(FMeshPaintParameters& InPaintParameters) {};
-	virtual void FinishPainting();
-	void UpdateResult();
-	double CalculateTargetEdgeLength(int TargetTriCount);
-	bool Paint(const FVector& InRayOrigin, const FVector& InRayDirection);
-	bool Paint(const TArrayView<TPair<FVector, FVector>>& Rays);
-	virtual void CacheSelectionData();
 	FPaintTexture2DData* GetPaintTargetData(const UTexture2D* InTexture);
 	FPaintTexture2DData* AddPaintTargetData(UTexture2D* InTexture);
-	void GatherTextureTriangles(IMeshPaintComponentAdapter* Adapter, int32 TriangleIndex, const int32 VertexIndices[3], TArray<FTexturePaintTriangleInfo>* TriangleInfo, TArray<FTexturePaintMeshSectionInfo>* SectionInfos, int32 UVChannelIndex);
-	void StartPaintingTexture(UMeshComponent* InMeshComponent, const IMeshPaintComponentAdapter& GeometryInfo);
-	void PaintTexture(FMeshPaintParameters& InParams, int32 UVChannel, TArray<FTexturePaintTriangleInfo>& InInfluencedTriangles, const IMeshPaintComponentAdapter& GeometryInfo, FMeshPaintParameters* LastParams = nullptr);
-	void FinishPaintingTexture();
-	void OnTransactionStateChanged(const FTransactionContext& InTransactionContext, const ETransactionStateEventType InTransactionState);
+	
+	void SetAllTextureOverrides();
+	void ClearAllTextureOverrides();
 
 	virtual UTexture2D* GetSelectedPaintTexture(UMeshComponent const* InMeshComponent) const { return nullptr; }
 	virtual int32 GetSelectedUVChannel(UMeshComponent const* InMeshComponent) const { return 0; }
 	virtual void CacheTexturePaintData() {}
+	virtual bool CanPaintTextureToComponent(UTexture* InTexture, UMeshComponent const* InMeshComponent) const { return false; }
 
 private:
-	bool PaintInternal(const TArrayView<TPair<FVector, FVector>>& Rays, EMeshPaintModeAction PaintAction, float PaintStrength);
-
+	void CacheSelectionData();
 	void AddTextureOverrideToComponent(FPaintTexture2DData& TextureData, UMeshComponent* MeshComponent, const IMeshPaintComponentAdapter* MeshPaintAdapter = nullptr);
+	double CalculateTargetEdgeLength(int TargetTriCount);
+	void StartPaintingTexture(UMeshComponent* InMeshComponent, const IMeshPaintComponentAdapter& GeometryInfo);
+	void UpdateResult();
+	void GatherTextureTriangles(IMeshPaintComponentAdapter* Adapter, int32 TriangleIndex, const int32 VertexIndices[3], TArray<FTexturePaintTriangleInfo>* TriangleInfo, TArray<FTexturePaintMeshSectionInfo>* SectionInfos, int32 UVChannelIndex);
+	bool Paint(const FVector& InRayOrigin, const FVector& InRayDirection);
+	bool Paint(const TArrayView<TPair<FVector, FVector>>& Rays);
+	void PaintTexture(FMeshPaintParameters& InParams, int32 UVChannel, TArray<FTexturePaintTriangleInfo>& InInfluencedTriangles, const IMeshPaintComponentAdapter& GeometryInfo, FMeshPaintParameters* LastParams = nullptr);
+	bool PaintInternal(const TArrayView<TPair<FVector, FVector>>& Rays, EMeshPaintModeAction PaintAction, float PaintStrength);
+	void FinishPaintingTexture();
+	void FinishPainting();
 
 protected:
 	UPROPERTY(Transient)
 	TObjectPtr<UMeshPaintSelectionMechanic> SelectionMechanic;
 
-	/** Textures eligible for painting retrieved from the current selection */
-	TArray<FPaintableTexture> PaintableTextures;
-
 	UPROPERTY(Transient)
 	TObjectPtr<UMeshTexturePaintingToolProperties> TextureProperties;
 
-	UPROPERTY(Transient)
-	TArray<TObjectPtr<const UTexture>> Textures;
+	/** Textures eligible for painting retrieved from the current selection */
+	TArray<FPaintableTexture> PaintableTextures;
 
 	/** Stores data associated with our paint target textures */
 	UPROPERTY(Transient)
 	TMap<TObjectPtr<UTexture2D>, FPaintTexture2DData> PaintTargetData;
 
-	/** Store the component overrides active for each paint target textures
-	 * Note this is not transactional because we use it as cache of the current state of the scene that we can clean/update after each transaction.
-	 */
-	UPROPERTY(Transient, NonTransactional)
-	TMap<TObjectPtr<UTexture2D>, FPaintComponentOverride> PaintComponentsOverride;
-
-	/** Texture paint: The mesh components that we're currently painting */
-	UPROPERTY(Transient)
-	TObjectPtr<UMeshComponent> TexturePaintingCurrentMeshComponent;
-
 	/** The original texture that we're painting */
 	UPROPERTY(Transient)
 	TObjectPtr<UTexture2D> PaintingTexture2D;
 
+	/** The mesh components that we're currently painting */
+	UPROPERTY(Transient)
+	TObjectPtr<UMeshComponent> TexturePaintingCurrentMeshComponent;
+
 	/** Hold the transaction while we are painting */
 	TUniquePtr<FScopedTransaction> PaintingTransaction;
 
-	double InitialMeshArea;
-	bool bResultValid;
-	bool bStampPending;
-	bool bInDrag;
+	double InitialMeshArea = 0;
+	bool bArePainting = false;
+	bool bResultValid = false;
+	bool bStampPending = false;
+	bool bInDrag = false;
+	bool bRequestPaintBucketFill = false;
+
+	bool bCachedClickRay = false;
 	FRay PendingStampRay;
 	FRay PendingClickRay;
 	FVector2D PendingClickScreenPosition;
-	bool bCachedClickRay;
-
+	
 	TArray<FPaintRayResults> LastPaintRayResults;
-	bool bRequestPaintBucketFill = false;
-
-	/** Flag for whether or not we are currently painting */
-	bool bArePainting;
-	bool bDoRestoreRenTargets;
-	/** Time kept since the user has started painting */
-	float TimeSinceStartedPainting;
-	/** Overall time value kept for drawing effects */
-	float Time;
 	FHitResult LastBestHitResult;
-	FSimpleDelegate OnPaintingFinishedDelegate;
-	/** Texture paint state */
-	/** Cached / stored instance texture paint settings for selected components */
-	TMap<UMeshComponent*, FInstanceTexturePaintSettings> ComponentToTexturePaintSettingsMap;
+
+	FOnPaintingFinishedDelegate OnPaintingFinishedDelegate;
 };
 
 /**
@@ -286,11 +255,9 @@ public:
 	virtual int32 GetSelectedUVChannel(UMeshComponent const* InMeshComponent) const override;
 	virtual void GetModifiedTexturesToSave(TArray<UObject*>& OutTexturesToSave) const override;
 	virtual void CacheTexturePaintData() override;
+	virtual bool CanPaintTextureToComponent(UTexture* InTexture, UMeshComponent const* InMeshComponent) const override;
 	// End UMeshTexturePaintingTool Interface.
 
-	/** Get mesh components that have mesh paint textures pending save and that should write vertex colors when they do save. */
-	void GetPaintedComponentsForPropagateVertexColor(TArray<UMeshComponent*>& OutComponents) const;
-
 protected:
 	UPROPERTY(Transient)
 	TObjectPtr<UMeshTextureColorPaintingToolProperties> ColorProperties;
@@ -316,21 +283,21 @@ public:
 	// Begin UMeshTexturePaintingTool Interface.
 	virtual bool AllowsMultiselect() const override { return false; }
 	virtual bool IsMeshAdapterSupported(TSharedPtr<IMeshPaintComponentAdapter> MeshAdapter) const override;
+	virtual void OnPropertyModified(UObject* PropertySet, FProperty* Property) override;
 	virtual UTexture2D* GetSelectedPaintTexture(UMeshComponent const* InMeshComponent) const override;
 	virtual int32 GetSelectedUVChannel(UMeshComponent const* InMeshComponent) const override;
 	virtual void GetModifiedTexturesToSave(TArray<UObject*>& OutTexturesToSave) const override;
 	virtual void CacheTexturePaintData() override;
+	virtual bool CanPaintTextureToComponent(UTexture* InTexture, UMeshComponent const* InMeshComponent) const override;
 	// End UMeshTexturePaintingTool Interface.
 
+	/** Change selected texture to previous or next available. */
+	void CycleTextures(int32 Direction);
 	/** Get the selected paint texture, and return the modified overriden texture if currently painting. */
 	UTexture* GetSelectedPaintTextureWithOverride() const;
 
-	/** Change selected texture to previous or next available. */
-	void CycleTextures(int32 Direction);
 	/** Returns true if asset shouldn't be shown in UI because it is not in our paintable texture array. */
 	bool ShouldFilterTextureAsset(const FAssetData& AssetData) const;
-	/** Call on change to selected paint texture */
-	void PaintTextureChanged(const FAssetData& AssetData);
 
 protected:
 	UPROPERTY(Transient)
diff --git a/Engine/Plugins/MeshPainting/Source/MeshPaintingToolset/Public/MeshVertexPaintingTool.h b/Engine/Plugins/MeshPainting/Source/MeshPaintingToolset/Public/MeshVertexPaintingTool.h
--- a/Engine/Plugins/MeshPainting/Source/MeshPaintingToolset/Public/MeshVertexPaintingTool.h
+++ b/Engine/Plugins/MeshPainting/Source/MeshPaintingToolset/Public/MeshVertexPaintingTool.h
@@ -211,9 +211,6 @@ private:
 
 	/** Flag for whether or not we are currently painting */
 	bool bArePainting;
-	bool bDoRestoreRenTargets;
-	/** Time kept since the user has started painting */
-	float TimeSinceStartedPainting;
 	/** Overall time value kept for drawing effects */
 	float Time;
 	FHitResult LastBestHitResult;
diff --git a/Engine/Plugins/MeshPainting/Source/MeshPaintingToolset/Public/TexturePaintToolset.h b/Engine/Plugins/MeshPainting/Source/MeshPaintingToolset/Public/TexturePaintToolset.h
--- a/Engine/Plugins/MeshPainting/Source/MeshPaintingToolset/Public/TexturePaintToolset.h
+++ b/Engine/Plugins/MeshPainting/Source/MeshPaintingToolset/Public/TexturePaintToolset.h
@@ -82,15 +82,9 @@ public:
 	/** Static: Creates a temporary texture used to transfer data to a render target in memory */
 	static UTexture2D* CreateScratchUncompressedTexture(UTexture2D* SourceTexture);
 
-	/** Makes sure that the render target is ready to paint on */
-	static void SetupInitialRenderTargetData(FPaintTexture2DData& PaintTextureData);
-
 	/** Keep old legacy method of initializing render target data for the paint brush texture; @todo MeshPaint: Migrate to the method with texture re-use */
 	static void SetupInitialRenderTargetData(UTexture2D* InTextureSource, UTextureRenderTarget2D* InRenderTarget);
 
-	/** Update the render target base on scratch texture */
-	static void UpdateRenderTargetData(FPaintTexture2DData& PaintTextureData);
-
 	/** Tries to find Materials using the given Texture and retrieve the corresponding material indices from the MEsh Compon*/
 	static void FindMaterialIndicesUsingTexture(const UTexture* Texture, const UMeshComponent* MeshComponent, TArray<int32>& OutIndices);
 	
@@ -102,7 +96,4 @@ public:
 
 	/** Retrieves all Paintable Textures from the given MeshComponent */
 	static void RetrieveTexturesForComponent(const UMeshComponent* Component, IMeshPaintComponentAdapter* Adapter, int32& OutDefaultIndex, TArray<FPaintableTexture>& OutTextures);
-
-	/** Checks whether or not the given MeshComponent's materials use the specified texture */
-	static bool DoesMeshComponentUseTexture(UMeshComponent* MeshComponent, UTexture* Texture);
 };
diff --git a/Engine/Source/Editor/EditorStyle/Private/StarshipStyle.cpp b/Engine/Source/Editor/EditorStyle/Private/StarshipStyle.cpp
--- a/Engine/Source/Editor/EditorStyle/Private/StarshipStyle.cpp
+++ b/Engine/Source/Editor/EditorStyle/Private/StarshipStyle.cpp
@@ -4949,12 +4949,10 @@ void FStarshipEditorStyle::FStyle::SetupLevelEditorStyle()
 		Set("MeshPaint.FillVertex.Small", new IMAGE_BRUSH("/Icons/GeneralTools/PaintBucket_40x", Icon20x20));
 		Set("MeshPaint.FillTexture", new IMAGE_BRUSH("/Icons/GeneralTools/PaintBucket_40x", Icon20x20));
 		Set("MeshPaint.FillTexture.Small", new IMAGE_BRUSH("/Icons/GeneralTools/PaintBucket_40x", Icon20x20));
-		Set("MeshPaint.ApplyVertex", new IMAGE_BRUSH("/Icons/Paint/Paint_Propagate_40x", Icon20x20));
-		Set("MeshPaint.ApplyVertex.Small", new IMAGE_BRUSH("/Icons/Paint/Paint_Propagate_40x", Icon20x20));
-		Set("MeshPaint.ApplyTexture", new IMAGE_BRUSH("/Icons/Paint/Paint_PropagateTexture_40x", Icon20x20));
-		Set("MeshPaint.ApplyTexture.Small", new IMAGE_BRUSH("/Icons/Paint/Paint_PropagateTexture_40x", Icon20x20));
-		Set("MeshPaint.Propagate", new IMAGE_BRUSH("Icons/Paint/Paint_AllLODs_40x", Icon20x20));
-		Set("MeshPaint.Propagate.Small", new IMAGE_BRUSH("Icons/Paint/Paint_AllLODs_40x", Icon20x20));
+		Set("MeshPaint.PropagateMesh", new IMAGE_BRUSH("/Icons/Paint/Paint_Propagate_40x", Icon20x20));
+		Set("MeshPaint.PropagateMesh.Small", new IMAGE_BRUSH("/Icons/Paint/Paint_Propagate_40x", Icon20x20));
+		Set("MeshPaint.PropagateLODs", new IMAGE_BRUSH("Icons/Paint/Paint_AllLODs_40x", Icon20x20));
+		Set("MeshPaint.PropagateLODs.Small", new IMAGE_BRUSH("Icons/Paint/Paint_AllLODs_40x", Icon20x20));
 		Set("MeshPaint.SaveVertex", new IMAGE_BRUSH("/Icons/GeneralTools/Save_40x", Icon20x20));
 		Set("MeshPaint.SaveVertex.Small", new IMAGE_BRUSH("/Icons/GeneralTools/Save_40x", Icon20x20));
 		Set("MeshPaint.SaveTexture", new IMAGE_BRUSH("/Icons/GeneralTools/Save_40x", Icon20x20));
