Non-unity build fix.
#rnx

diff --git a/Engine/Plugins/MeshPainting/Source/MeshPaintEditorMode/Private/MeshPaintModeHelpers.cpp b/Engine/Plugins/MeshPainting/Source/MeshPaintEditorMode/Private/MeshPaintModeHelpers.cpp
--- a/Engine/Plugins/MeshPainting/Source/MeshPaintEditorMode/Private/MeshPaintModeHelpers.cpp
+++ b/Engine/Plugins/MeshPainting/Source/MeshPaintEditorMode/Private/MeshPaintModeHelpers.cpp
@@ -35,6 +35,7 @@
 #include "Factories/FbxSkeletalMeshImportData.h"
 
 #include "Rendering/SkeletalMeshModel.h"
+#include "MeshPaintHelpers.h"
 #include "MeshPaintMode.h"
 #include "MeshPaintVisualize.h"
 #include "MeshTexturePaintingTool.h"
