Prevent crashes in MeshPaintMode when operating on cooked editor packages.

diff --git a/Engine/Plugins/MeshPainting/Source/MeshPaintEditorMode/Private/MeshPaintModeHelpers.cpp b/Engine/Plugins/MeshPainting/Source/MeshPaintEditorMode/Private/MeshPaintModeHelpers.cpp
--- a/Engine/Plugins/MeshPainting/Source/MeshPaintEditorMode/Private/MeshPaintModeHelpers.cpp
+++ b/Engine/Plugins/MeshPainting/Source/MeshPaintEditorMode/Private/MeshPaintModeHelpers.cpp
@@ -668,6 +668,13 @@ bool UMeshPaintModeSubsystem::CanPropagateVertexColors(TArray<UStaticMeshCompone
 
 		if (StaticMesh != nullptr)
 		{
+			// Disallow propagation of vertex colors to cooked static mesh assets.
+			if (StaticMesh->GetOutermost()->bIsCookedForEditor)
+			{
+				bValid = false;
+				break;
+			}
+			
 			StaticMeshes.AddUnique(StaticMesh);
 		}
 
diff --git a/Engine/Plugins/MeshPainting/Source/MeshPaintingToolset/Private/MeshPaintHelpers.cpp b/Engine/Plugins/MeshPainting/Source/MeshPaintingToolset/Private/MeshPaintHelpers.cpp
--- a/Engine/Plugins/MeshPainting/Source/MeshPaintingToolset/Private/MeshPaintHelpers.cpp
+++ b/Engine/Plugins/MeshPainting/Source/MeshPaintingToolset/Private/MeshPaintHelpers.cpp
@@ -125,10 +125,15 @@ void UMeshPaintingSubsystem::RemoveComponentMeshPaintTexture(UStaticMeshComponen
 
 bool UMeshPaintingSubsystem::PropagateColorsToRawMesh(UStaticMesh* StaticMesh, int32 LODIndex, FStaticMeshComponentLODInfo& ComponentLODInfo)
 {
-	check(ComponentLODInfo.OverrideVertexColors);
-	check(StaticMesh->IsSourceModelValid(LODIndex));
-	check(StaticMesh->GetRenderData());
-	check(StaticMesh->GetRenderData()->LODResources.IsValidIndex(LODIndex));
+	if (!ComponentLODInfo.OverrideVertexColors ||
+		!StaticMesh ||
+		!StaticMesh->IsSourceModelValid(LODIndex) ||
+		!StaticMesh->GetRenderData() ||
+		!StaticMesh->GetRenderData()->LODResources.IsValidIndex(LODIndex) ||
+		StaticMesh->GetOutermost()->bIsCookedForEditor)
+	{
+		return false;
+	}
 
 	bool bPropagatedColors = false;
 	FStaticMeshSourceModel& SrcModel = StaticMesh->GetSourceModel(LODIndex);
