Vertex weight painting supports per LOD painting.

diff --git a/Engine/Plugins/MeshPainting/Source/MeshPaintEditorMode/Private/BrushSettingsCustomization.cpp b/Engine/Plugins/MeshPainting/Source/MeshPaintEditorMode/Private/BrushSettingsCustomization.cpp
--- a/Engine/Plugins/MeshPainting/Source/MeshPaintEditorMode/Private/BrushSettingsCustomization.cpp
+++ b/Engine/Plugins/MeshPainting/Source/MeshPaintEditorMode/Private/BrushSettingsCustomization.cpp
@@ -17,6 +17,7 @@
 #include "MeshPaintHelpers.h"
 #include "MeshPaintModeHelpers.h"
 #include "MeshTexturePaintingTool.h"
+#include "MeshVertexPaintingTool.h"
 #include "ContentBrowserDelegates.h"
 #include "GeometryCollection/GeometryCollectionComponent.h"
 
@@ -179,6 +180,115 @@ TSharedRef<IDetailCustomization> FVertexPaintingSettingsCustomization::MakeInsta
 
 void FVertexPaintingSettingsCustomization::CustomizeDetails(IDetailLayoutBuilder& DetailLayout)
 {
+	IDetailCategoryBuilder& VertexCategory = DetailLayout.EditCategory(TEXT("VertexPainting"));
+
+	/** Add custom row for painting on specific LOD level with callbacks to the painter to update the data */
+	TSharedRef<IPropertyHandle> LODPaintingEnabled = DetailLayout.GetProperty("bPaintOnSpecificLOD", UMeshVertexPaintingToolProperties::StaticClass());
+	LODPaintingEnabled->MarkHiddenByCustomization();
+	TSharedRef<IPropertyHandle> LODPaintingIndex = DetailLayout.GetProperty("LODIndex", UMeshVertexPaintingToolProperties::StaticClass());
+	LODPaintingIndex->MarkHiddenByCustomization();
+	TSharedPtr<SWidget> LODIndexWidget = LODPaintingIndex->CreatePropertyValueWidget();
+
+	VertexCategory.AddCustomRow(NSLOCTEXT("LODPainting", "LODPaintingLabel", "LOD Model Painting"))
+	.NameContent()
+	[
+		SNew(STextBlock)
+		.Text(NSLOCTEXT("LODPainting", "LODPaintingSetupLabel", "LOD Model Painting"))
+		.Font(IDetailLayoutBuilder::GetDetailFont())
+		.ToolTipText(NSLOCTEXT("LODPainting", "LODPaintingSetupToolTip", "Allows for Painting Vertex Colors on Specific LOD Models."))
+	]
+	.ValueContent()
+	[
+		SNew(SHorizontalBox)
+		+ SHorizontalBox::Slot()
+		.Padding(0.0f, 0.0f, 4.0f, 0.0f)
+		.AutoWidth()
+		[
+			SNew(SCheckBox)
+			.IsChecked_Lambda([=]() -> ECheckBoxState 
+			{
+				return (UMeshPaintMode::GetVertexToolProperties() && UMeshPaintMode::GetVertexToolProperties()->bPaintOnSpecificLOD ? ECheckBoxState::Checked : ECheckBoxState::Unchecked); 
+			})
+			.OnCheckStateChanged(FOnCheckStateChanged::CreateLambda([=](ECheckBoxState State) 
+			{ 
+				if (UMeshVertexPaintingTool* VertexBrush = Cast<UMeshVertexPaintingTool>(UMeshPaintMode::GetMeshPaintMode()->GetToolManager()->GetActiveTool(EToolSide::Left)))
+				{
+					VertexBrush->LODPaintStateChanged(State == ECheckBoxState::Checked);
+				}
+			}))
+		]
+		+ SHorizontalBox::Slot()
+		.Padding(0.0f, 0.0f, 4.0f, 0.0f)
+		[
+			SNew(SNumericEntryBox<int32>)
+			.Font(IDetailLayoutBuilder::GetDetailFont())
+			.IsEnabled_Lambda([=]() -> bool 
+			{
+				return UMeshPaintMode::GetVertexToolProperties() ? UMeshPaintMode::GetVertexToolProperties()->bPaintOnSpecificLOD : false;  
+			})
+			.AllowSpin(true)
+			.Value_Lambda([=]() -> int32 
+			{
+				return UMeshPaintMode::GetVertexToolProperties() ? UMeshPaintMode::GetVertexToolProperties()->LODIndex : 0; 
+			})
+			.MinValue(0)
+			.MaxValue_Lambda([=]() -> int32 
+			{ 
+				if (UMeshVertexPaintingTool* VertexBrush = Cast<UMeshVertexPaintingTool>(UMeshPaintMode::GetMeshPaintMode()->GetToolManager()->GetActiveTool(EToolSide::Left)))
+				{
+					return VertexBrush->GetMaxLODIndexToPaint();
+				}
+				return INT_MAX;
+			})
+			.MaxSliderValue_Lambda([=]() -> int32 
+			{ 
+				if (UMeshVertexPaintingTool* VertexBrush = Cast<UMeshVertexPaintingTool>(UMeshPaintMode::GetMeshPaintMode()->GetToolManager()->GetActiveTool(EToolSide::Left)))
+				{
+					return VertexBrush->GetMaxLODIndexToPaint();
+				}
+				return INT_MAX;
+			})
+			.OnValueChanged(SNumericEntryBox<int32>::FOnValueChanged::CreateLambda([=](int32 Value) { UMeshPaintMode::GetVertexToolProperties()->LODIndex = Value; }))
+			.OnValueCommitted(SNumericEntryBox<int32>::FOnValueCommitted::CreateLambda([=](int32 Value, ETextCommit::Type CommitType)
+			{ 
+				UMeshPaintMode::GetVertexToolProperties()->LODIndex = Value; 
+				if (UMeshVertexPaintingTool* VertexBrush = Cast<UMeshVertexPaintingTool>(UMeshPaintMode::GetMeshPaintMode()->GetToolManager()->GetActiveTool(EToolSide::Left)))
+				{
+					VertexBrush->PaintLODChanged();
+				}
+			}))
+		]
+	];
+
+	VertexCategory.AddCustomRow(NSLOCTEXT("LODPainting", "LODPaintingLabel", "LOD Model Painting"))
+	.WholeRowContent()
+	[
+		SNew(SWarningOrErrorBox)
+		.Visibility_Lambda([this]() -> EVisibility
+		{
+			if (UMeshVertexPaintingToolProperties* VertexProperties = UMeshPaintMode::GetVertexToolProperties())
+			{
+				return VertexProperties->bPaintOnSpecificLOD ? EVisibility::Collapsed : EVisibility::Visible;
+			}
+			return EVisibility::Collapsed;
+		})
+		.Message_Lambda([this]() -> FText
+		{
+			static const FText SkelMeshNotificationText = LOCTEXT("SkelMeshAssetPaintInfo", "Paint is propagated to Skeletal Mesh Asset(s)");
+			static const FText StaticMeshNotificationText = LOCTEXT("StaticMeshAssetPaintInfo", "Paint is applied to all LODs");
+			static const FText GeometryCollectionNotificationText = LOCTEXT("GeometryCollectionAssetPaintInfo", "Paint is propagated to Geometry Collection Asset(s), and Geometry Collection does not currently support LODs.");
+
+			const bool bGeometryCollectionText = UMeshPaintMode::GetMeshPaintMode()->GetSelectedComponents<UGeometryCollectionComponent>().Num() > 0;
+			const bool bSkelMeshText = UMeshPaintMode::GetMeshPaintMode()->GetSelectedComponents<USkeletalMeshComponent>().Num() > 0;
+			const bool bLODPaintText = UMeshPaintMode::GetVertexToolProperties() ?  !UMeshPaintMode::GetVertexToolProperties()->bPaintOnSpecificLOD : false;
+			return FText::Format(FTextFormat::FromString(TEXT("{0}{1}{2}{3}")), 
+				bSkelMeshText ? SkelMeshNotificationText : FText::GetEmpty(),
+				bGeometryCollectionText ? GeometryCollectionNotificationText : FText::GetEmpty(),
+				bSkelMeshText && bLODPaintText ? FText::FromString(TEXT("\n")) : FText::GetEmpty(),
+				bLODPaintText ? StaticMeshNotificationText : FText::GetEmpty()
+				);
+		})
+	];
 }
 
 
@@ -194,7 +304,6 @@ void FVertexColorPaintingSettingsCustomization::CustomizeDetails(IDetailLayoutBu
 
 	IDetailCategoryBuilder& ColorCategory = DetailLayout.EditCategory(TEXT("ColorPainting"));
 	ColorCategory.SetSortOrder(0);
-	IDetailCategoryBuilder& VertexCategory = DetailLayout.EditCategory(TEXT("VertexPainting"));
 
 	/** Creates a custom widget row containing all color channel flags */
 	TSharedRef<IPropertyHandle> RedChannel = DetailLayout.GetProperty("bWriteRed", UMeshVertexColorPaintingToolProperties::StaticClass());
@@ -209,17 +318,18 @@ void FVertexColorPaintingSettingsCustomization::CustomizeDetails(IDetailLayoutBu
 	TSharedPtr<SHorizontalBox> ChannelsWidget;
 
 	ColorCategory.AddCustomRow(NSLOCTEXT("VertexPaintSettings", "ChannelLabel", "Channels"))
-		.NameContent()
-		[
-			SNew(STextBlock)
-			.Text(NSLOCTEXT("VertexPaintSettings", "ChannelsLabel", "Channels"))
-			.ToolTipText(NSLOCTEXT("VertexPaintSettings", "ChannelsToolTip", "Colors Channels which should be influenced during Painting."))
-		]
-		.ValueContent()
-		.MaxDesiredWidth(250.0f)
-		[
-			SAssignNew(ChannelsWidget, SHorizontalBox)
-		];
+	.NameContent()
+	[
+		SNew(STextBlock)
+		.Text(NSLOCTEXT("VertexPaintSettings", "ChannelsLabel", "Channels"))
+		.ToolTipText(NSLOCTEXT("VertexPaintSettings", "ChannelsToolTip", "Colors Channels which should be influenced during Painting."))
+		.Font(IDetailLayoutBuilder::GetDetailFont())
+	]
+	.ValueContent()
+	.MaxDesiredWidth(250.0f)
+	[
+		SAssignNew(ChannelsWidget, SHorizontalBox)
+	];
 
 	for (TSharedRef<IPropertyHandle> Channel : Channels)
 	{
@@ -230,102 +340,6 @@ void FVertexColorPaintingSettingsCustomization::CustomizeDetails(IDetailLayoutBu
 			CreateColorChannelWidget(Channel)
 		];
 	}
-
-	static FText RestrictReason = NSLOCTEXT("VertexPaintSettings", "TextureIndexRestriction", "Unable to paint this Texture, change Texture Weight Type");
-	BlendPaintEnumRestriction = MakeShareable(new FPropertyRestriction(RestrictReason));
-
-	/** Add custom row for painting on specific LOD level with callbacks to the painter to update the data */
-	TSharedRef<IPropertyHandle> LODPaintingEnabled = DetailLayout.GetProperty("bPaintOnSpecificLOD", UMeshVertexColorPaintingToolProperties::StaticClass());
-	LODPaintingEnabled->MarkHiddenByCustomization();
-	TSharedRef<IPropertyHandle> LODPaintingIndex = DetailLayout.GetProperty("LODIndex", UMeshVertexColorPaintingToolProperties::StaticClass());
-	LODPaintingIndex->MarkHiddenByCustomization();
-	TSharedPtr<SWidget> LODIndexWidget = LODPaintingIndex->CreatePropertyValueWidget();
-
-	VertexCategory.AddCustomRow(NSLOCTEXT("LODPainting", "LODPaintingLabel", "LOD Model Painting"))
-		.NameContent()
-		[
-			SNew(STextBlock)
-			.Text(NSLOCTEXT("LODPainting", "LODPaintingSetupLabel", "LOD Model Painting"))
-			.ToolTipText(NSLOCTEXT("LODPainting", "LODPaintingSetupToolTip", "Allows for Painting Vertex Colors on Specific LOD Models."))
-		]
-		.ValueContent()
-		[
-			SNew(SHorizontalBox)
-			+ SHorizontalBox::Slot()
-		.Padding(0.0f, 0.0f, 4.0f, 0.0f)
-		.AutoWidth()
-		[
-			SNew(SCheckBox)
-				.IsChecked_Lambda([=]() -> ECheckBoxState { return (UMeshPaintMode::GetVertexColorToolProperties() && UMeshPaintMode::GetVertexColorToolProperties()->bPaintOnSpecificLOD ? ECheckBoxState::Checked : ECheckBoxState::Unchecked); })
-				.OnCheckStateChanged(FOnCheckStateChanged::CreateLambda([=](ECheckBoxState State) { 
-					if (UMeshColorPaintingTool* ColorBrush = Cast<UMeshColorPaintingTool>(UMeshPaintMode::GetMeshPaintMode()->GetToolManager()->GetActiveTool(EToolSide::Left)))
-					{
-						ColorBrush->LODPaintStateChanged(State == ECheckBoxState::Checked);
-					}
-				}))
-		]
-		+ SHorizontalBox::Slot()
-		.Padding(0.0f, 0.0f, 4.0f, 0.0f)
-		[
-			SNew(SNumericEntryBox<int32>)
-			.IsEnabled_Lambda([=]() -> bool { return UMeshPaintMode::GetVertexColorToolProperties() ? UMeshPaintMode::GetVertexColorToolProperties()->bPaintOnSpecificLOD : false;  })
-			.AllowSpin(true)
-			.Value_Lambda([=]() -> int32 { return UMeshPaintMode::GetVertexColorToolProperties() ? UMeshPaintMode::GetVertexColorToolProperties()->LODIndex : 0; })
-			.MinValue(0)
-			.MaxValue_Lambda([=]() -> int32 { 
-					if (UMeshColorPaintingTool* ColorBrush = Cast<UMeshColorPaintingTool>(UMeshPaintMode::GetMeshPaintMode()->GetToolManager()->GetActiveTool(EToolSide::Left)))
-					{
-						return ColorBrush->GetMaxLODIndexToPaint();
-					}
-					return INT_MAX;
-				})
-			.MaxSliderValue_Lambda([=]() -> int32 { 
-					if (UMeshColorPaintingTool* ColorBrush = Cast<UMeshColorPaintingTool>(UMeshPaintMode::GetMeshPaintMode()->GetToolManager()->GetActiveTool(EToolSide::Left)))
-					{
-						return ColorBrush->GetMaxLODIndexToPaint();
-					}
-					return INT_MAX;
-				})
-			.OnValueChanged(SNumericEntryBox<int32>::FOnValueChanged::CreateLambda([=](int32 Value) { UMeshPaintMode::GetVertexColorToolProperties()->LODIndex = Value; }))
-			.OnValueCommitted(SNumericEntryBox<int32>::FOnValueCommitted::CreateLambda([=](int32 Value, ETextCommit::Type CommitType) { 
-					UMeshPaintMode::GetVertexColorToolProperties()->LODIndex = Value; 
-					if (UMeshColorPaintingTool* ColorBrush = Cast<UMeshColorPaintingTool>(UMeshPaintMode::GetMeshPaintMode()->GetToolManager()->GetActiveTool(EToolSide::Left)))
-					{
-						ColorBrush->PaintLODChanged();
-					}
-				}))
-		]
-		];
-
-	VertexCategory.AddCustomRow(NSLOCTEXT("LODPainting", "LODPaintingLabel", "LOD Model Painting"))
-		.WholeRowContent()
-		[
-			SNew(SWarningOrErrorBox)
-			.Visibility_Lambda([this]() -> EVisibility
-			{
-				if (UMeshVertexColorPaintingToolProperties* ColorProperties = UMeshPaintMode::GetVertexColorToolProperties())
-				{
-					return ColorProperties->bPaintOnSpecificLOD ? EVisibility::Collapsed : EVisibility::Visible;
-				}
-				return EVisibility::Collapsed;
-			})
-			.Message_Lambda([this]() -> FText
-			{
-				static const FText SkelMeshNotificationText = LOCTEXT("SkelMeshAssetPaintInfo", "Paint is propagated to Skeletal Mesh Asset(s)");
-				static const FText StaticMeshNotificationText = LOCTEXT("StaticMeshAssetPaintInfo", "Paint is applied to all LODs");
-				static const FText GeometryCollectionNotificationText = LOCTEXT("GeometryCollectionAssetPaintInfo", "Paint is propagated to Geometry Collection Asset(s), and Geometry Collection does not currently support LODs.");
-
-				const bool bGeometryCollectionText = UMeshPaintMode::GetMeshPaintMode()->GetSelectedComponents<UGeometryCollectionComponent>().Num() > 0;
-				const bool bSkelMeshText = UMeshPaintMode::GetMeshPaintMode()->GetSelectedComponents<USkeletalMeshComponent>().Num() > 0;
-				const bool bLODPaintText = UMeshPaintMode::GetVertexColorToolProperties() ?  !UMeshPaintMode::GetVertexColorToolProperties()->bPaintOnSpecificLOD : false;
-				return FText::Format(FTextFormat::FromString(TEXT("{0}{1}{2}{3}")), 
-					bSkelMeshText ? SkelMeshNotificationText : FText::GetEmpty(),
-					bGeometryCollectionText ? GeometryCollectionNotificationText : FText::GetEmpty(),
-					bSkelMeshText && bLODPaintText ? FText::FromString(TEXT("\n")) : FText::GetEmpty(),
-					bLODPaintText ? StaticMeshNotificationText : FText::GetEmpty()
-					);
-			})
-		];
 }
 
 
diff --git a/Engine/Plugins/MeshPainting/Source/MeshPaintEditorMode/Private/BrushSettingsCustomization.h b/Engine/Plugins/MeshPainting/Source/MeshPaintEditorMode/Private/BrushSettingsCustomization.h
--- a/Engine/Plugins/MeshPainting/Source/MeshPaintEditorMode/Private/BrushSettingsCustomization.h
+++ b/Engine/Plugins/MeshPainting/Source/MeshPaintEditorMode/Private/BrushSettingsCustomization.h
@@ -36,10 +36,6 @@ public:
 
 	/** IPropertyTypeCustomization interface */
 	virtual void CustomizeDetails(IDetailLayoutBuilder& DetailLayout) override;
-
-protected:
-	/** Property restriction applied to blend paint enum dropdown box */
-	TSharedPtr<FPropertyRestriction> BlendPaintEnumRestriction;
 };
 
 class FVertexColorPaintingSettingsCustomization : public FVertexPaintingSettingsCustomization
@@ -60,6 +56,9 @@ public:
 	virtual void CustomizeDetails(IDetailLayoutBuilder& DetailLayout) override;
 
 protected:
+	/** Property restriction applied to blend paint enum dropdown box */
+	TSharedPtr<FPropertyRestriction> BlendPaintEnumRestriction;
+
 	/** Callback for when texture weight type changed so we can update restrictions */
 	void OnTextureWeightTypeChanged(TSharedRef<IPropertyHandle> WeightTypeProperty, TSharedRef<IPropertyHandle> PaintWeightProperty, TSharedRef<IPropertyHandle> EraseWeightProperty);
 };
diff --git a/Engine/Plugins/MeshPainting/Source/MeshPaintEditorMode/Private/MeshPaintMode.cpp b/Engine/Plugins/MeshPainting/Source/MeshPaintEditorMode/Private/MeshPaintMode.cpp
--- a/Engine/Plugins/MeshPainting/Source/MeshPaintEditorMode/Private/MeshPaintMode.cpp
+++ b/Engine/Plugins/MeshPainting/Source/MeshPaintEditorMode/Private/MeshPaintMode.cpp
@@ -22,6 +22,7 @@
 #include "MeshPaintModeHelpers.h"
 #include "MeshSelect.h"
 #include "MeshTexturePaintingTool.h"
+#include "MeshVertexPaintingTool.h"
 #include "Modules/ModuleManager.h"
 #include "LevelEditor.h"
 #include "TexturePaintToolset.h"
@@ -46,13 +47,20 @@ FString UMeshPaintMode::TextureColorPaintToolName = TEXT("TextureColorBrushTool"
 FString UMeshPaintMode::TextureAssetPaintToolName = TEXT("TextureBrushTool");
 
 
+UMeshPaintMode* UMeshPaintMode::GetMeshPaintMode()
+{
+	return Cast<UMeshPaintMode>(GLevelEditorModeTools().GetActiveScriptableMode("MeshPaintMode"));
+}
+
 template<typename T>
 static T* GetTypedToolProperties()
 {
-	if (GLevelEditorModeTools().GetActiveScriptableMode("MeshPaintMode")
-		&& GLevelEditorModeTools().GetActiveScriptableMode("MeshPaintMode")->GetToolManager()->GetActiveTool(EToolSide::Mouse))
+	UMeshPaintMode* MeshPaintMode = UMeshPaintMode::GetMeshPaintMode();
+	UInteractiveToolManager* ToolManager = MeshPaintMode != nullptr ? MeshPaintMode->GetToolManager() : nullptr;
+	UInteractiveTool* Tool = ToolManager != nullptr ? ToolManager->GetActiveTool(EToolSide::Mouse) : nullptr;
+	if (Tool != nullptr)
 	{
-		TArray<UObject*> PropertyArray = GLevelEditorModeTools().GetActiveScriptableMode("MeshPaintMode")->GetToolManager()->GetActiveTool(EToolSide::Mouse)->GetToolProperties();
+		const TArray<UObject*> PropertyArray = Tool->GetToolProperties();
 		for (UObject* Property : PropertyArray)
 		{
 			if (T* FoundProperty = Cast<T>(Property))
@@ -99,11 +107,6 @@ UMeshTextureAssetPaintingToolProperties* UMeshPaintMode::GetTextureAssetToolProp
 	return GetTypedToolProperties<UMeshTextureAssetPaintingToolProperties>();
 }
 
-UMeshPaintMode* UMeshPaintMode::GetMeshPaintMode()
-{
-	return Cast<UMeshPaintMode>(GLevelEditorModeTools().GetActiveScriptableMode("MeshPaintMode"));
-}
-
 
 UMeshPaintMode::UMeshPaintMode()
 	: Super()
@@ -138,10 +141,10 @@ void UMeshPaintMode::Enter()
 	UTextureAssetAdapterClickToolBuilder* TextureAssetClickToolBuilder = NewObject<UTextureAssetAdapterClickToolBuilder>(this);
 	RegisterTool(ToolManagerCommands.SelectTextureAsset, TextureAssetSelectToolName, TextureAssetClickToolBuilder);
 
-	UMeshColorPaintingToolBuilder* MeshColorPaintingToolBuilder = NewObject<UMeshColorPaintingToolBuilder>(this);
+	UMeshVertexColorPaintingToolBuilder* MeshColorPaintingToolBuilder = NewObject<UMeshVertexColorPaintingToolBuilder>(this);
 	RegisterTool(ToolManagerCommands.PaintVertexColor, VertexColorPaintToolName, MeshColorPaintingToolBuilder);
 
-	UMeshWeightPaintingToolBuilder* WeightPaintingToolBuilder = NewObject<UMeshWeightPaintingToolBuilder>(this);
+	UMeshVertexWeightPaintingToolBuilder* WeightPaintingToolBuilder = NewObject<UMeshVertexWeightPaintingToolBuilder>(this);
 	RegisterTool(ToolManagerCommands.PaintVertexWeight, VertexWeightPaintToolName, WeightPaintingToolBuilder);
 
 	UMeshTextureColorPaintingToolBuilder* MeshTextureColorPaintingToolBuilder = NewObject<UMeshTextureColorPaintingToolBuilder>(this);
@@ -378,9 +381,9 @@ void UMeshPaintMode::OnToolStarted(UInteractiveToolManager* Manager, UInteractiv
 
 void UMeshPaintMode::OnVertexPaintFinished()
 {
-	if (UMeshVertexColorPaintingToolProperties* ColorPaintingToolProperties = UMeshPaintMode::GetVertexColorToolProperties())
+	if (UMeshVertexPaintingToolProperties* VertexPaintingToolProperties = UMeshPaintMode::GetVertexToolProperties())
 	{
-		if (!ColorPaintingToolProperties->bPaintOnSpecificLOD)
+		if (!VertexPaintingToolProperties->bPaintOnSpecificLOD)
 		{
 			PropagateVertexColorsToLODs();
 		}
@@ -389,7 +392,6 @@ void UMeshPaintMode::OnVertexPaintFinished()
 			if (UMeshPaintingSubsystem* MeshPaintingSubsystem = GEngine->GetEngineSubsystem<UMeshPaintingSubsystem>())
 			{
 				bRecacheDataSizes = true;
-
 				MeshPaintingSubsystem->Refresh();
 			}
 		}
@@ -495,7 +497,7 @@ void UMeshPaintMode::FillVertexColors()
 	FColor FillColor = FColor::White;
 	FColor MaskColor = FColor::White;
 
-	if (GetToolManager()->GetActiveTool(EToolSide::Mouse)->IsA<UMeshWeightPaintingTool>())
+	if (GetToolManager()->GetActiveTool(EToolSide::Mouse)->IsA<UMeshVertexWeightPaintingTool>())
 	{
 		FillColor = MeshPaintingSubsystem->GenerateColorForTextureWeight((int32)GetVertexWeightToolProperties()->TextureWeightType, (int32)GetVertexWeightToolProperties()->PaintTextureWeightIndex).ToFColor(bConvertSRGB);
 	}
@@ -522,17 +524,16 @@ void UMeshPaintMode::FillVertexColors()
 			MeshAdapter->PreEdit();
 		}
 
-		UMeshVertexColorPaintingToolProperties* ColorProperties = UMeshPaintMode::GetVertexColorToolProperties();
-		
-		const bool bPaintOnSpecificLOD = ColorProperties ? ColorProperties->bPaintOnSpecificLOD : false;
+		UMeshVertexPaintingToolProperties* VertexProperties = UMeshPaintMode::GetVertexToolProperties();
+		const bool bPaintOnSpecificLOD = VertexProperties ? VertexProperties->bPaintOnSpecificLOD : false;
 
 		if (Component->IsA<UStaticMeshComponent>())
 		{
-			MeshPaintingSubsystem->FillStaticMeshVertexColors(Cast<UStaticMeshComponent>(Component), bPaintOnSpecificLOD ? ColorProperties->LODIndex : -1, FillColor, MaskColor);
+			MeshPaintingSubsystem->FillStaticMeshVertexColors(Cast<UStaticMeshComponent>(Component), bPaintOnSpecificLOD ? VertexProperties->LODIndex : -1, FillColor, MaskColor);
 		}
 		else if (Component->IsA<USkeletalMeshComponent>())
 		{
-			GEngine->GetEngineSubsystem<UMeshPaintingSubsystem>()->FillSkeletalMeshVertexColors(Cast<USkeletalMeshComponent>(Component), bPaintOnSpecificLOD ? ColorProperties->LODIndex : -1, FillColor, MaskColor);
+			GEngine->GetEngineSubsystem<UMeshPaintingSubsystem>()->FillSkeletalMeshVertexColors(Cast<USkeletalMeshComponent>(Component), bPaintOnSpecificLOD ? VertexProperties->LODIndex : -1, FillColor, MaskColor);
 		}
 		else if (MeshAdapter) 			// We don't have a custom fill function for this type of component; try to go through the adapter.
 		{
@@ -661,9 +662,9 @@ bool UMeshPaintMode::CanRemoveInstanceVertexColors() const
 {
 	const TArray<UStaticMeshComponent*> StaticMeshComponents = GetSelectedComponents<UStaticMeshComponent>();
 	int32 PaintingMeshLODIndex = 0;
-	if (UMeshVertexColorPaintingToolProperties* ColorProperties = UMeshPaintMode::GetVertexColorToolProperties())
+	if (UMeshVertexPaintingToolProperties* VertexProperties = UMeshPaintMode::GetVertexToolProperties())
 	{
-		PaintingMeshLODIndex = ColorProperties->bPaintOnSpecificLOD ? ColorProperties->LODIndex : 0;
+		PaintingMeshLODIndex = VertexProperties->bPaintOnSpecificLOD ? VertexProperties->LODIndex : 0;
 	}
 	int32 NumValidMeshes = 0;
 	// Retrieve per instance vertex color information (only valid if the component contains actual instance vertex colors)
@@ -694,9 +695,9 @@ bool UMeshPaintMode::CanCopyInstanceVertexColors() const
 {
 	const TArray<UStaticMeshComponent*> StaticMeshComponents = GetSelectedComponents<UStaticMeshComponent>();
 	int32 PaintingMeshLODIndex = 0;
-	if (UMeshVertexColorPaintingToolProperties* ColorProperties = UMeshPaintMode::GetVertexColorToolProperties())
+	if (UMeshVertexPaintingToolProperties* VertexProperties = UMeshPaintMode::GetVertexToolProperties())
 	{
-		PaintingMeshLODIndex = ColorProperties->bPaintOnSpecificLOD ? ColorProperties->LODIndex : 0;
+		PaintingMeshLODIndex = VertexProperties->bPaintOnSpecificLOD ? VertexProperties->LODIndex : 0;
 	}
 
 	return GEditor->GetEditorSubsystem<UMeshPaintModeSubsystem>()->CanCopyInstanceVertexColors(StaticMeshComponents, PaintingMeshLODIndex);
@@ -706,9 +707,9 @@ bool UMeshPaintMode::CanCopyInstanceVertexColors() const
 bool UMeshPaintMode::CanPropagateVertexColorsToLODs() const
 {
 	bool bPaintOnSpecificLOD = false;
-	if (UMeshVertexColorPaintingToolProperties* ColorProperties = UMeshPaintMode::GetVertexColorToolProperties())
+	if (UMeshVertexPaintingToolProperties* VertexProperties = UMeshPaintMode::GetVertexToolProperties())
 	{
-		bPaintOnSpecificLOD = ColorProperties ? ColorProperties->bPaintOnSpecificLOD : false;
+		bPaintOnSpecificLOD = VertexProperties ? VertexProperties->bPaintOnSpecificLOD : false;
 	}
 	// Can propagate when the mesh contains per-lod vertex colors or when we are not painting to a specific lod
 	const bool bSelectionContainsPerLODColors = GEngine->GetEngineSubsystem<UMeshPaintingSubsystem>()->SelectionContainsPerLODColors();
@@ -827,19 +828,16 @@ void UMeshPaintMode::PropagateVertexColorsToLODs()
 template<typename ComponentClass>
 TArray<ComponentClass*> UMeshPaintMode::GetSelectedComponents() const
 {
-	TArray<ComponentClass*> Components;
 	FToolBuilderState SelectionState;
 	GetToolManager()->GetContextQueriesAPI()->GetCurrentSelectionState(SelectionState);
-//	if (PaintSettings->PaintMode == EPaintMode::Textures)
-	{
 
-		for (int32 SelectionIndex = 0; SelectionIndex < SelectionState.SelectedComponents.Num(); ++SelectionIndex)
+	TArray<ComponentClass*> Components;
+	for (int32 SelectionIndex = 0; SelectionIndex < SelectionState.SelectedComponents.Num(); ++SelectionIndex)
+	{
+		ComponentClass* SelectedComponent = Cast<ComponentClass>(SelectionState.SelectedComponents[SelectionIndex]);
+		if (SelectedComponent)
 		{
-			ComponentClass* SelectedComponent = Cast<ComponentClass>(SelectionState.SelectedComponents[SelectionIndex]);
-			if (SelectedComponent)
-			{
-				Components.AddUnique(SelectedComponent);
-			}
+			Components.AddUnique(SelectedComponent);
 		}
 	}
 
@@ -893,14 +891,12 @@ void UMeshPaintMode::UpdateCachedDataSizes()
 	bRecacheDataSizes = false;
 }
 
-
-
 void UMeshPaintMode::CycleMeshLODs(int32 Direction)
 {
-	if (UMeshColorPaintingTool* ColorPaintingTool = Cast<UMeshColorPaintingTool>(GetToolManager()->GetActiveTool(EToolSide::Left)))
+	if (UMeshVertexPaintingTool* VertexPaintingTool = Cast<UMeshVertexPaintingTool>(GetToolManager()->GetActiveTool(EToolSide::Left)))
 	{
 		FScopedTransaction Transaction(LOCTEXT("LevelMeshPainter_CycleLOD", "Changed Current LOD"));
-		ColorPaintingTool->CycleMeshLODs(Direction);
+		VertexPaintingTool->CycleMeshLODs(Direction);
 	}
 }
 
@@ -919,7 +915,6 @@ bool UMeshPaintMode::CanCycleTextures() const
 	return TexturePaintingTool != nullptr;
 }
 
-
 void UMeshPaintMode::ActivateDefaultTool()
 {
 	FName PaletteName = Toolkit->GetCurrentPalette();
diff --git a/Engine/Plugins/MeshPainting/Source/MeshPaintEditorMode/Private/MeshPaintMode.h b/Engine/Plugins/MeshPainting/Source/MeshPaintEditorMode/Private/MeshPaintMode.h
--- a/Engine/Plugins/MeshPainting/Source/MeshPaintEditorMode/Private/MeshPaintMode.h
+++ b/Engine/Plugins/MeshPainting/Source/MeshPaintEditorMode/Private/MeshPaintMode.h
@@ -3,13 +3,13 @@
 #pragma once
 
 #include "Tools/UEdMode.h"
-#include "MeshVertexPaintingTool.h"
 #include "Tools/LegacyEdModeInterfaces.h"
 #include "MeshPaintMode.generated.h"
 
 class UMeshPaintingToolProperties;
 class UMeshVertexPaintingToolProperties;
 class UMeshVertexColorPaintingToolProperties;
+class UMeshVertexWeightPaintingToolProperties;
 class UMeshTexturePaintingToolProperties;
 class UMeshTextureColorPaintingToolProperties;
 class UMeshTextureAssetPaintingToolProperties;
diff --git a/Engine/Plugins/MeshPainting/Source/MeshPaintEditorMode/Private/MeshPaintModeCommands.cpp b/Engine/Plugins/MeshPainting/Source/MeshPaintEditorMode/Private/MeshPaintModeCommands.cpp
--- a/Engine/Plugins/MeshPainting/Source/MeshPaintEditorMode/Private/MeshPaintModeCommands.cpp
+++ b/Engine/Plugins/MeshPainting/Source/MeshPaintEditorMode/Private/MeshPaintModeCommands.cpp
@@ -1,14 +1,17 @@
 // Copyright Epic Games, Inc. All Rights Reserved.
 
 #include "MeshPaintModeCommands.h"
+
 #include "MeshPaintMode.h"
+#include "MeshTexturePaintingTool.h"
+#include "MeshVertexPaintingTool.h"
 
 #define LOCTEXT_NAMESPACE "MeshPaintEditorModeCommands"
 
 void FMeshPaintingToolActionCommands::GetToolDefaultObjectList(TArray<UInteractiveTool*>& ToolCDOs)
 {
 	ToolCDOs.Add(GetMutableDefault<USingleSelectionTool>());
-	ToolCDOs.Add(GetMutableDefault<UMeshColorPaintingTool>());
+	ToolCDOs.Add(GetMutableDefault<UMeshVertexColorPaintingTool>());
 }
 
 void FMeshPaintingToolActionCommands::RegisterAllToolActions()
diff --git a/Engine/Plugins/MeshPainting/Source/MeshPaintEditorMode/Private/MeshPaintModeHelpers.cpp b/Engine/Plugins/MeshPainting/Source/MeshPaintEditorMode/Private/MeshPaintModeHelpers.cpp
--- a/Engine/Plugins/MeshPainting/Source/MeshPaintEditorMode/Private/MeshPaintModeHelpers.cpp
+++ b/Engine/Plugins/MeshPainting/Source/MeshPaintEditorMode/Private/MeshPaintModeHelpers.cpp
@@ -39,6 +39,7 @@
 #include "MeshPaintMode.h"
 #include "MeshPaintVisualize.h"
 #include "MeshTexturePaintingTool.h"
+#include "MeshVertexPaintingTool.h"
 #include "InterchangeAssetImportData.h"
 #include "InterchangeGenericAssetsPipeline.h"
 #include "InterchangePythonPipelineBase.h"
@@ -679,7 +680,7 @@ bool UMeshPaintModeSubsystem::CanPropagateVertexColors(TArray<UStaticMeshCompone
 		}
 
 		int32 CachedLODIndex = 0;
-		if (UMeshColorPaintingTool* ColorPaintingTool = Cast<UMeshColorPaintingTool>(UMeshPaintMode::GetMeshPaintMode()->GetToolManager()->GetActiveTool(EToolSide::Left)))
+		if (UMeshVertexColorPaintingTool* ColorPaintingTool = Cast<UMeshVertexColorPaintingTool>(UMeshPaintMode::GetMeshPaintMode()->GetToolManager()->GetActiveTool(EToolSide::Left)))
 		{
 			CachedLODIndex = ColorPaintingTool->GetCachedLODIndex();
 		}
diff --git a/Engine/Plugins/MeshPainting/Source/MeshPaintingToolset/Private/MeshVertexPaintingTool.cpp b/Engine/Plugins/MeshPainting/Source/MeshPaintingToolset/Private/MeshVertexPaintingTool.cpp
--- a/Engine/Plugins/MeshPainting/Source/MeshPaintingToolset/Private/MeshVertexPaintingTool.cpp
+++ b/Engine/Plugins/MeshPainting/Source/MeshPaintingToolset/Private/MeshVertexPaintingTool.cpp
@@ -18,25 +18,25 @@
  * ToolBuilder
  */
 
-bool UMeshColorPaintingToolBuilder::CanBuildTool(const FToolBuilderState& SceneState) const
+bool UMeshVertexColorPaintingToolBuilder::CanBuildTool(const FToolBuilderState& SceneState) const
 {
 	return GEngine->GetEngineSubsystem<UMeshPaintingSubsystem>()->GetSelectionSupportsVertexPaint();
 }
 
-UInteractiveTool* UMeshColorPaintingToolBuilder::BuildTool(const FToolBuilderState& SceneState) const
+UInteractiveTool* UMeshVertexColorPaintingToolBuilder::BuildTool(const FToolBuilderState& SceneState) const
 {
-	UMeshColorPaintingTool* NewTool = NewObject<UMeshColorPaintingTool>(SceneState.ToolManager);
+	UMeshVertexColorPaintingTool* NewTool = NewObject<UMeshVertexColorPaintingTool>(SceneState.ToolManager);
 	return NewTool;
 }
 
-bool UMeshWeightPaintingToolBuilder::CanBuildTool(const FToolBuilderState& SceneState) const
+bool UMeshVertexWeightPaintingToolBuilder::CanBuildTool(const FToolBuilderState& SceneState) const
 {
 	return GEngine->GetEngineSubsystem<UMeshPaintingSubsystem>()->GetSelectionSupportsVertexPaint();
 }
 
-UInteractiveTool* UMeshWeightPaintingToolBuilder::BuildTool(const FToolBuilderState& SceneState) const
+UInteractiveTool* UMeshVertexWeightPaintingToolBuilder::BuildTool(const FToolBuilderState& SceneState) const
 {
-	UMeshWeightPaintingTool* NewTool = NewObject<UMeshWeightPaintingTool>(SceneState.ToolManager);
+	UMeshVertexWeightPaintingTool* NewTool = NewObject<UMeshVertexWeightPaintingTool>(SceneState.ToolManager);
 	return NewTool;
 }
 
@@ -84,6 +84,10 @@ void UMeshVertexPaintingTool::Setup()
 
 void UMeshVertexPaintingTool::Shutdown(EToolShutdownType ShutdownType)
 {
+	//If we're painting vertex colors then propagate the painting done on LOD0 to all lower LODs. 
+	//Then stop forcing the LOD level of the mesh to LOD0.
+	ApplyForcedLODIndex(-1);
+
 	FinishPainting();
 	BrushProperties->SaveProperties(this);
 	UMeshPaintingSubsystem* MeshPaintingSubsystem = GEngine->GetEngineSubsystem<UMeshPaintingSubsystem>();
@@ -417,16 +421,6 @@ void UMeshVertexPaintingTool::UpdateResult()
 	bResultValid = true;
 }
 
-bool UMeshVertexPaintingTool::HasAccept() const
-{
-	return false;
-}
-
-bool UMeshVertexPaintingTool::CanAccept() const
-{
-	return false;
-}
-
 FInputRayHit UMeshVertexPaintingTool::CanBeginClickDragSequence(const FInputDeviceRay& PressPos)
 {
 	FHitResult OutHit;
@@ -518,12 +512,12 @@ void UMeshVertexPaintingTool::FinishPainting()
 }
 
 
-UMeshColorPaintingTool::UMeshColorPaintingTool()
+UMeshVertexColorPaintingTool::UMeshVertexColorPaintingTool()
 {
 	PropertyClass = UMeshVertexColorPaintingToolProperties::StaticClass();
 }
 
-void UMeshColorPaintingTool::Setup()
+void UMeshVertexColorPaintingTool::Setup()
 {
 	Super::Setup();
 	ColorProperties = Cast<UMeshVertexColorPaintingToolProperties>(BrushProperties);
@@ -533,26 +527,17 @@ void UMeshColorPaintingTool::Setup()
 		EToolMessageLevel::UserNotification);
 }
 
-void UMeshColorPaintingTool::Shutdown(EToolShutdownType ShutdownType)
+void UMeshVertexPaintingTool::CacheSelectionData()
 {
-	//If we're painting vertex colors then propagate the painting done on LOD0 to all lower LODs. 
-	//Then stop forcing the LOD level of the mesh to LOD0.
-	ApplyForcedLODIndex(-1);
-	Super::Shutdown(ShutdownType);
-}
-
-void UMeshColorPaintingTool::CacheSelectionData()
-{
-	Super::CacheSelectionData();
 	if (UMeshPaintingSubsystem* MeshPaintingSubsystem = GEngine->GetEngineSubsystem<UMeshPaintingSubsystem>())
 	{
 		MeshPaintingSubsystem->ClearPaintableMeshComponents();
 		// Update(cached) Paint LOD level if necessary
-		ColorProperties->LODIndex = FMath::Min<int32>(ColorProperties->LODIndex, GetMaxLODIndexToPaint());
-		CachedLODIndex = ColorProperties->LODIndex;
-		bCachedForceLOD = ColorProperties->bPaintOnSpecificLOD;
+		VertexProperties->LODIndex = FMath::Min<int32>(VertexProperties->LODIndex, GetMaxLODIndexToPaint());
+		CachedLODIndex = VertexProperties->LODIndex;
+		bCachedForceLOD = VertexProperties->bPaintOnSpecificLOD;
 		//Determine LOD level to use for painting(can only paint on LODs in vertex mode)
-		const int32 PaintLODIndex = ColorProperties->bPaintOnSpecificLOD ? ColorProperties->LODIndex : 0;
+		const int32 PaintLODIndex = VertexProperties->bPaintOnSpecificLOD ? VertexProperties->LODIndex : 0;
 		//Determine UV channel to use while painting textures
 		const int32 UVChannel = 0;
 
@@ -560,7 +545,7 @@ void UMeshColorPaintingTool::CacheSelectionData()
 	}
 }
 
-void UMeshColorPaintingTool::SetAdditionalPaintParameters(FMeshPaintParameters& InPaintParameters)
+void UMeshVertexColorPaintingTool::SetAdditionalPaintParameters(FMeshPaintParameters& InPaintParameters)
 {
 	InPaintParameters.bWriteRed = ColorProperties->bWriteRed;
 	InPaintParameters.bWriteGreen = ColorProperties->bWriteGreen;
@@ -569,7 +554,7 @@ void UMeshColorPaintingTool::SetAdditionalPaintParameters(FMeshPaintParameters&
 	InPaintParameters.ApplyVertexDataDelegate.AddUObject(GEngine->GetEngineSubsystem<UMeshPaintingSubsystem>(), &UMeshPaintingSubsystem::ApplyVertexColorPaint);
 }
 
-int32 UMeshColorPaintingTool::GetMaxLODIndexToPaint() const
+int32 UMeshVertexPaintingTool::GetMaxLODIndexToPaint() const
 {
 	//The maximum LOD we can paint is decide by the lowest number of LOD in the selection
 	int32 LODMin = TNumericLimits<int32>::Max();
@@ -595,18 +580,17 @@ int32 UMeshColorPaintingTool::GetMaxLODIndexToPaint() const
 }
 
 
-void UMeshColorPaintingTool::LODPaintStateChanged(const bool bLODPaintingEnabled)
+void UMeshVertexPaintingTool::LODPaintStateChanged(const bool bLODPaintingEnabled)
 {
 	bool AbortChange = false;
 
 	// Set actual flag in the settings struct
-	ColorProperties->bPaintOnSpecificLOD = bLODPaintingEnabled;
+	VertexProperties->bPaintOnSpecificLOD = bLODPaintingEnabled;
 
 	if (!bLODPaintingEnabled)
 	{
 		// Reset painting LOD index
-		ColorProperties->LODIndex = 0;
-
+		VertexProperties->LODIndex = 0;
 	}
 
 	ApplyForcedLODIndex(bLODPaintingEnabled ? CachedLODIndex : -1);
@@ -629,7 +613,7 @@ void UMeshColorPaintingTool::LODPaintStateChanged(const bool bLODPaintingEnabled
 	}
 }
 
-void UMeshColorPaintingTool::ApplyForcedLODIndex(int32 ForcedLODIndex)
+void UMeshVertexPaintingTool::ApplyForcedLODIndex(int32 ForcedLODIndex)
 {
 	if (UMeshPaintingSubsystem* MeshPaintingSubsystem = GEngine->GetEngineSubsystem<UMeshPaintingSubsystem>())
 	{
@@ -646,12 +630,12 @@ void UMeshColorPaintingTool::ApplyForcedLODIndex(int32 ForcedLODIndex)
 }
 
 
-void UMeshColorPaintingTool::PaintLODChanged()
+void UMeshVertexPaintingTool::PaintLODChanged()
 {
 	// Enforced LOD for painting
-	if (CachedLODIndex != ColorProperties->LODIndex)
+	if (CachedLODIndex != VertexProperties->LODIndex)
 	{
-		CachedLODIndex = ColorProperties->LODIndex;
+		CachedLODIndex = VertexProperties->LODIndex;
 		ApplyForcedLODIndex(bCachedForceLOD ? CachedLODIndex : -1);
 
 		TUniquePtr< FComponentReregisterContext > ComponentReregisterContext;
@@ -673,14 +657,14 @@ void UMeshColorPaintingTool::PaintLODChanged()
 	}
 }
 
-void UMeshColorPaintingTool::CycleMeshLODs(int32 Direction)
+void UMeshVertexPaintingTool::CycleMeshLODs(int32 Direction)
 {
 	if (bCachedForceLOD)
 	{
 		const int32 MaxLODIndex = GetMaxLODIndexToPaint() + 1;
-		const int32 NewLODIndex = ColorProperties->LODIndex + Direction;
+		const int32 NewLODIndex = VertexProperties->LODIndex + Direction;
 		const int32 AdjustedLODIndex = NewLODIndex < 0 ? MaxLODIndex + NewLODIndex : NewLODIndex % MaxLODIndex;
-		ColorProperties->LODIndex = AdjustedLODIndex;
+		VertexProperties->LODIndex = AdjustedLODIndex;
 		PaintLODChanged();
 	}
 }
@@ -694,12 +678,12 @@ UMeshVertexWeightPaintingToolProperties::UMeshVertexWeightPaintingToolProperties
 
 }
 
-UMeshWeightPaintingTool::UMeshWeightPaintingTool()
+UMeshVertexWeightPaintingTool::UMeshVertexWeightPaintingTool()
 {
 	PropertyClass = UMeshVertexWeightPaintingToolProperties::StaticClass();
 }
 
-void UMeshWeightPaintingTool::Setup()
+void UMeshVertexWeightPaintingTool::Setup()
 {
 	Super::Setup();
 	WeightProperties = Cast<UMeshVertexWeightPaintingToolProperties>(BrushProperties);
@@ -709,21 +693,7 @@ void UMeshWeightPaintingTool::Setup()
 		EToolMessageLevel::UserNotification);
 }
 
-void UMeshWeightPaintingTool::CacheSelectionData()
-{
-	Super::CacheSelectionData();
-	if (UMeshPaintingSubsystem* MeshPaintingSubsystem = GEngine->GetEngineSubsystem<UMeshPaintingSubsystem>())
-	{
-		MeshPaintingSubsystem->ClearPaintableMeshComponents();
-		//Determine LOD level to use for painting(can only paint on LODs in vertex mode)
-		const int32 PaintLODIndex = 0;
-		//Determine UV channel to use while painting textures
-		const int32 UVChannel = 0;
-		MeshPaintingSubsystem->CacheSelectionData(PaintLODIndex, UVChannel);
-	}
-}
-
-void UMeshWeightPaintingTool::SetAdditionalPaintParameters(FMeshPaintParameters& InPaintParameters)
+void UMeshVertexWeightPaintingTool::SetAdditionalPaintParameters(FMeshPaintParameters& InPaintParameters)
 {
 	InPaintParameters.TotalWeightCount = (int32)WeightProperties->TextureWeightType;
 
diff --git a/Engine/Plugins/MeshPainting/Source/MeshPaintingToolset/Public/MeshTexturePaintingTool.h b/Engine/Plugins/MeshPainting/Source/MeshPaintingToolset/Public/MeshTexturePaintingTool.h
--- a/Engine/Plugins/MeshPainting/Source/MeshPaintingToolset/Public/MeshTexturePaintingTool.h
+++ b/Engine/Plugins/MeshPainting/Source/MeshPaintingToolset/Public/MeshTexturePaintingTool.h
@@ -102,7 +102,7 @@ class MESHPAINTINGTOOLSET_API UMeshTextureColorPaintingToolProperties : public U
 	GENERATED_BODY()
 
 public:
-	/** Whether to copy texture color painting to vertex colors on Apply. */
+	/** Whether to copy all texture color painting to vertex colors. */
 	UPROPERTY(EditAnywhere, Category = ColorPainting)
 	bool bPropagateToVertexColor = false;
 };
@@ -137,6 +137,16 @@ class MESHPAINTINGTOOLSET_API UMeshTexturePaintingTool : public UBaseBrushTool,
 public:
 	UMeshTexturePaintingTool();
 
+	DECLARE_DELEGATE_OneParam(FOnPaintingFinishedDelegate, UMeshComponent*);
+	FOnPaintingFinishedDelegate& OnPaintingFinished() { return OnPaintingFinishedDelegate; }
+
+	void FloodCurrentPaintTexture();
+
+	virtual void GetModifiedTexturesToSave(TArray<UObject*>& OutTexturesToSave) const {}
+	virtual int32 GetSelectedUVChannel(UMeshComponent const* InMeshComponent) const { return 0; }
+
+protected:
+	// Begin UInteractiveTool Interface.
 	virtual void Setup() override;
 	virtual void Shutdown(EToolShutdownType ShutdownType) override;
 	virtual void Render(IToolsContextRenderAPI* RenderAPI) override;
@@ -152,16 +162,8 @@ public:
 	virtual	bool HitTest(const FRay& Ray, FHitResult& OutHit) override;
 	virtual void OnPropertyModified(UObject* PropertySet, FProperty* Property) override;
 	virtual double EstimateMaximumTargetDimension() override;
-	virtual bool IsPainting() const { return bArePainting; }
-
-	DECLARE_DELEGATE_OneParam(FOnPaintingFinishedDelegate, UMeshComponent*);
-	FOnPaintingFinishedDelegate& OnPaintingFinished() { return OnPaintingFinishedDelegate; }
-
-	void FloodCurrentPaintTexture();
-	
-	virtual void GetModifiedTexturesToSave(TArray<UObject*>& OutTexturesToSave) const {}
+	// End UInteractiveTool Interface.
 
-protected:
 	FPaintTexture2DData* GetPaintTargetData(const UTexture2D* InTexture);
 	FPaintTexture2DData* AddPaintTargetData(UTexture2D* InTexture);
 	
@@ -169,7 +171,6 @@ protected:
 	void ClearAllTextureOverrides();
 
 	virtual UTexture2D* GetSelectedPaintTexture(UMeshComponent const* InMeshComponent) const { return nullptr; }
-	virtual int32 GetSelectedUVChannel(UMeshComponent const* InMeshComponent) const { return 0; }
 	virtual void CacheTexturePaintData() {}
 	virtual bool CanPaintTextureToComponent(UTexture* InTexture, UMeshComponent const* InMeshComponent) const { return false; }
 
@@ -188,12 +189,6 @@ private:
 	void FinishPainting();
 
 protected:
-	UPROPERTY(Transient)
-	TObjectPtr<UMeshPaintSelectionMechanic> SelectionMechanic;
-
-	UPROPERTY(Transient)
-	TObjectPtr<UMeshTexturePaintingToolProperties> TextureProperties;
-
 	/** Textures eligible for painting retrieved from the current selection */
 	TArray<FPaintableTexture> PaintableTextures;
 
@@ -201,11 +196,18 @@ protected:
 	UPROPERTY(Transient)
 	TMap<TObjectPtr<UTexture2D>, FPaintTexture2DData> PaintTargetData;
 
+private:
+	UPROPERTY(Transient)
+	TObjectPtr<UMeshPaintSelectionMechanic> SelectionMechanic;
+
+	UPROPERTY(Transient)
+	TObjectPtr<UMeshTexturePaintingToolProperties> TextureProperties;
+
 	/** The original texture that we're painting */
 	UPROPERTY(Transient)
 	TObjectPtr<UTexture2D> PaintingTexture2D;
 
-	/** The mesh components that we're currently painting */
+	/** The mesh component that we're currently painting on */
 	UPROPERTY(Transient)
 	TObjectPtr<UMeshComponent> TexturePaintingCurrentMeshComponent;
 
@@ -244,6 +246,7 @@ class MESHPAINTINGTOOLSET_API UMeshTextureColorPaintingTool : public UMeshTextur
 public:
 	UMeshTextureColorPaintingTool();
 
+protected:
 	// Begin UInteractiveTool Interface.
 	virtual void Setup() override;
 	// End UInteractiveTool Interface.
@@ -258,7 +261,7 @@ public:
 	virtual bool CanPaintTextureToComponent(UTexture* InTexture, UMeshComponent const* InMeshComponent) const override;
 	// End UMeshTexturePaintingTool Interface.
 
-protected:
+private:
 	UPROPERTY(Transient)
 	TObjectPtr<UMeshTextureColorPaintingToolProperties> ColorProperties;
 };
@@ -276,6 +279,18 @@ class MESHPAINTINGTOOLSET_API UMeshTextureAssetPaintingTool : public UMeshTextur
 public:
 	UMeshTextureAssetPaintingTool();
 	
+	/** Change selected texture to previous or next available. */
+	void CycleTextures(int32 Direction);
+
+	/** Get the selected paint texture, and return the modified overriden texture if currently painting. */
+	UTexture* GetSelectedPaintTextureWithOverride() const;
+
+	/** Returns true if asset shouldn't be shown in UI because it is not in our paintable texture array. */
+	bool ShouldFilterTextureAsset(const FAssetData& AssetData) const;
+
+	virtual int32 GetSelectedUVChannel(UMeshComponent const* InMeshComponent) const override;
+
+protected:
 	// Begin UInteractiveTool Interface.
 	virtual void Setup() override;
 	// End UInteractiveTool Interface.
@@ -285,21 +300,12 @@ public:
 	virtual bool IsMeshAdapterSupported(TSharedPtr<IMeshPaintComponentAdapter> MeshAdapter) const override;
 	virtual void OnPropertyModified(UObject* PropertySet, FProperty* Property) override;
 	virtual UTexture2D* GetSelectedPaintTexture(UMeshComponent const* InMeshComponent) const override;
-	virtual int32 GetSelectedUVChannel(UMeshComponent const* InMeshComponent) const override;
 	virtual void GetModifiedTexturesToSave(TArray<UObject*>& OutTexturesToSave) const override;
 	virtual void CacheTexturePaintData() override;
 	virtual bool CanPaintTextureToComponent(UTexture* InTexture, UMeshComponent const* InMeshComponent) const override;
 	// End UMeshTexturePaintingTool Interface.
 
-	/** Change selected texture to previous or next available. */
-	void CycleTextures(int32 Direction);
-	/** Get the selected paint texture, and return the modified overriden texture if currently painting. */
-	UTexture* GetSelectedPaintTextureWithOverride() const;
-
-	/** Returns true if asset shouldn't be shown in UI because it is not in our paintable texture array. */
-	bool ShouldFilterTextureAsset(const FAssetData& AssetData) const;
-
-protected:
+private:
 	UPROPERTY(Transient)
 	TObjectPtr<UMeshTextureAssetPaintingToolProperties> AssetProperties;
 };
diff --git a/Engine/Plugins/MeshPainting/Source/MeshPaintingToolset/Public/MeshVertexPaintingTool.h b/Engine/Plugins/MeshPainting/Source/MeshPaintingToolset/Public/MeshVertexPaintingTool.h
--- a/Engine/Plugins/MeshPainting/Source/MeshPaintingToolset/Public/MeshVertexPaintingTool.h
+++ b/Engine/Plugins/MeshPainting/Source/MeshPaintingToolset/Public/MeshVertexPaintingTool.h
@@ -8,12 +8,10 @@
 #include "MeshPaintInteractions.h"
 #include "MeshVertexPaintingTool.generated.h"
 
+enum class EMeshPaintModeAction : uint8;
 enum class EToolShutdownType : uint8;
 struct FPerVertexPaintActionArgs;
-
-
 struct FToolBuilderState;
-enum class EMeshPaintModeAction : uint8;
 class IMeshPaintComponentAdapter;
 
 struct FPaintRayResults
@@ -48,11 +46,12 @@
 	TextureFive
 };
 
+
 /**
  *
  */
 UCLASS()
-class MESHPAINTINGTOOLSET_API UMeshColorPaintingToolBuilder : public UInteractiveToolBuilder
+class MESHPAINTINGTOOLSET_API UMeshVertexColorPaintingToolBuilder : public UInteractiveToolBuilder
 {
 	GENERATED_BODY()
 
@@ -62,7 +61,7 @@
 };
 
 UCLASS()
-class MESHPAINTINGTOOLSET_API UMeshWeightPaintingToolBuilder : public UInteractiveToolBuilder
+class MESHPAINTINGTOOLSET_API UMeshVertexWeightPaintingToolBuilder : public UInteractiveToolBuilder
 {
 	GENERATED_BODY()
 
@@ -72,7 +71,6 @@
 };
 
 
-
 UCLASS()
 class MESHPAINTINGTOOLSET_API UMeshVertexPaintingToolProperties : public UMeshPaintingToolProperties
 {
@@ -81,12 +79,19 @@
 public:
 	UMeshVertexPaintingToolProperties();
 
+	/** When unchecked the painting on the base LOD will be propagate automatically to all other LODs when exiting the mode or changing the selection */
+	UPROPERTY(EditAnywhere, Category = VertexPainting, meta = (InlineEditConditionToggle, TransientToolProperty))
+	bool bPaintOnSpecificLOD = false;
+
+	/** Index of LOD to paint. If not set then paint is applied to all LODs. */
+	UPROPERTY(EditAnywhere, Category = VertexPainting, meta = (UIMin = "0", ClampMin = "0", EditCondition = "bPaintOnSpecificLOD", TransientToolProperty))
+	int32 LODIndex = 0;
+
 	/** Size of vertex points drawn when mesh painting is active. */
-	UPROPERTY(EditAnywhere, AdvancedDisplay, Category = "VertexPainting")
+	UPROPERTY(EditAnywhere, Category = "VertexPainting")
 	float VertexPreviewSize;
 };
 
-
 UCLASS()
 class MESHPAINTINGTOOLSET_API UMeshVertexColorPaintingToolProperties : public UMeshVertexPaintingToolProperties
 {
@@ -108,14 +113,6 @@
 	/** Whether or not to apply Vertex Color Painting to the Alpha Channel */
 	UPROPERTY(EditAnywhere, Category = ColorPainting, DisplayName = "Alpha")
 	bool bWriteAlpha = false;
-
-	/** When unchecked the painting on the base LOD will be propagate automatically to all other LODs when exiting the mode or changing the selection */
-	UPROPERTY(EditAnywhere, Category = VertexPainting, meta = (InlineEditConditionToggle, TransientToolProperty))
-	bool bPaintOnSpecificLOD = false;
-
-	/** Index of LOD to paint. If not set then paint is applied to all LODs. */
-	UPROPERTY(EditAnywhere, Category = VertexPainting, meta = (UIMin = "0", ClampMin = "0", EditCondition = "bPaintOnSpecificLOD", TransientToolProperty))
-	int32 LODIndex = 0;
 };
 
 UCLASS()
@@ -139,6 +136,7 @@
 	EMeshPaintTextureIndex EraseTextureWeightIndex;
 };
 
+
 UCLASS(Abstract)
 class MESHPAINTINGTOOLSET_API UMeshVertexPaintingTool : public UBaseBrushTool, public IMeshPaintSelectionInterface
 {
@@ -147,13 +145,24 @@
 public:
 	UMeshVertexPaintingTool();
 
+	void PaintLODChanged();
+	void LODPaintStateChanged(const bool bLODPaintingEnabled);
+
+	int32 GetMaxLODIndexToPaint() const;
+	int32 GetCachedLODIndex() const { return CachedLODIndex; }
+
+	void CycleMeshLODs(int32 Direction);
+
+	FSimpleDelegate& OnPaintingFinished() { return OnPaintingFinishedDelegate; }
+
+protected:
 	virtual void Setup() override;
 	virtual void Shutdown(EToolShutdownType ShutdownType) override;
 	virtual void Render(IToolsContextRenderAPI* RenderAPI) override;
 	virtual void OnTick(float DeltaTime) override;
 	virtual bool HasCancel() const override { return false; }
-	virtual bool HasAccept() const override;
-	virtual bool CanAccept() const override;
+	virtual bool HasAccept() const override { return false; }
+	virtual bool CanAccept() const override { return false; }
 	virtual FInputRayHit CanBeginClickDragSequence(const FInputDeviceRay& PressPos) override;
 	virtual	void OnUpdateModifierState(int ModifierID, bool bIsOn) override;
 	virtual void OnBeginDrag(const FRay& Ray) override;
@@ -161,107 +170,79 @@
 	virtual void OnEndDrag(const FRay& Ray) override;
 	virtual	bool HitTest(const FRay& Ray, FHitResult& OutHit) override;
 	virtual void OnPropertyModified(UObject* PropertySet, FProperty* Property) override;
-	virtual bool IsPainting() const
-	{
-		return bArePainting;
-	}
-	virtual double EstimateMaximumTargetDimension() override;
-
-	FSimpleDelegate& OnPaintingFinished()
-	{
-		return OnPaintingFinishedDelegate;
-	}
+	virtual bool AllowsMultiselect() const override { return true; }
 	virtual bool IsMeshAdapterSupported(TSharedPtr<IMeshPaintComponentAdapter> MeshAdapter) const override;
-	virtual bool AllowsMultiselect() const override
-	{
-		return true;
-	}
+	virtual double EstimateMaximumTargetDimension() override;
 
-protected:
 	virtual void SetAdditionalPaintParameters(FMeshPaintParameters& InPaintParameters) {};
-	virtual void FinishPainting();
+
+private:
+	void CacheSelectionData();
+	void ApplyForcedLODIndex(int32 ForcedLODIndex);
 	void UpdateResult();
-	double CalculateTargetEdgeLength(int TargetTriCount);
 	bool Paint(const FVector& InRayOrigin, const FVector& InRayDirection);
 	bool Paint(const TArrayView<TPair<FVector, FVector>>& Rays);
-	virtual void CacheSelectionData() {};
-	/** Per vertex action function used for painting vertex data */
+	bool PaintInternal(const TArrayView<TPair<FVector, FVector>>& Rays, EMeshPaintModeAction PaintAction, float PaintStrength);
 	void ApplyVertexData(FPerVertexPaintActionArgs& InArgs, int32 VertexIndex, FMeshPaintParameters Parameters);
+	void FinishPainting();
+	double CalculateTargetEdgeLength(int TargetTriCount);
 
+private:
+	UPROPERTY(Transient)
+	TObjectPtr<UMeshPaintSelectionMechanic> SelectionMechanic;
 
-protected:
-	double InitialMeshArea;
-	bool bResultValid;
-	bool bStampPending;
-	bool bInDrag;
+	UPROPERTY(Transient)
+	TObjectPtr<UMeshVertexPaintingToolProperties> VertexProperties;
+
+	/** Current LOD index used for painting / forcing */
+	int32 CachedLODIndex;
+	/** Whether or not a specific LOD level should be forced */
+	bool bCachedForceLOD;
+
+	double InitialMeshArea = 0;
+	bool bArePainting = false;
+	bool bResultValid = false;
+	bool bStampPending = false;
+	bool bInDrag = false;
+	
+	bool bCachedClickRay = false;
 	FRay PendingStampRay;
 	FRay PendingClickRay;
 	FVector2D PendingClickScreenPosition;
-	bool bCachedClickRay;
-
-	UPROPERTY(Transient)
-	TObjectPtr<UMeshPaintSelectionMechanic> SelectionMechanic;
-
-private:
-	bool PaintInternal(const TArrayView<TPair<FVector, FVector>>& Rays, EMeshPaintModeAction PaintAction, float PaintStrength);
-	
-private:
-	UPROPERTY(Transient)
-	TObjectPtr<UMeshVertexPaintingToolProperties> VertexProperties;
-
-	/** Flag for whether or not we are currently painting */
-	bool bArePainting;
-	/** Overall time value kept for drawing effects */
-	float Time;
 	FHitResult LastBestHitResult;
+	
 	FSimpleDelegate OnPaintingFinishedDelegate;
 };
 
 UCLASS()
-class MESHPAINTINGTOOLSET_API UMeshColorPaintingTool : public UMeshVertexPaintingTool
+class MESHPAINTINGTOOLSET_API UMeshVertexColorPaintingTool : public UMeshVertexPaintingTool
 {
 	GENERATED_BODY()
 
 public:
-	UMeshColorPaintingTool();
-	virtual void Setup() override;
-	virtual void Shutdown(EToolShutdownType ShutdownType) override;
-	void LODPaintStateChanged(const bool bLODPaintingEnabled);
-	int32 GetMaxLODIndexToPaint() const;
-	void PaintLODChanged();
-	int32 GetCachedLODIndex() const
-	{
-		return CachedLODIndex;
-	}
-	void CycleMeshLODs(int32 Direction);
+	UMeshVertexColorPaintingTool();
 
 protected:
-	virtual void CacheSelectionData() override;
+	virtual void Setup() override;
+
 	virtual void SetAdditionalPaintParameters(FMeshPaintParameters& InPaintParameters) override;
-	void ApplyForcedLODIndex(int32 ForcedLODIndex);
-
+	
 private:
 	UPROPERTY(Transient)
 	TObjectPtr<UMeshVertexColorPaintingToolProperties> ColorProperties;
-
-	/** Current LOD index used for painting / forcing */
-	int32 CachedLODIndex;
-	/** Whether or not a specific LOD level should be forced */
-	bool bCachedForceLOD;
 };
 
-
 UCLASS()
-class MESHPAINTINGTOOLSET_API UMeshWeightPaintingTool : public UMeshVertexPaintingTool
+class MESHPAINTINGTOOLSET_API UMeshVertexWeightPaintingTool : public UMeshVertexPaintingTool
 {
 	GENERATED_BODY()
 
 public:
-	UMeshWeightPaintingTool();
-	virtual void Setup() override;
+	UMeshVertexWeightPaintingTool();
 
 protected:
-	virtual void CacheSelectionData() override;
+	virtual void Setup() override;
+	
 	virtual void SetAdditionalPaintParameters(FMeshPaintParameters& InPaintParameters);
 
 private:
