Fix mesh not visually updating when setting paint LOD in MeshPaintTool.
#rnx

diff --git a/Engine/Plugins/MeshPainting/Source/MeshPaintingToolset/Private/MeshPaintHelpers.cpp b/Engine/Plugins/MeshPainting/Source/MeshPaintingToolset/Private/MeshPaintHelpers.cpp
--- a/Engine/Plugins/MeshPainting/Source/MeshPaintingToolset/Private/MeshPaintHelpers.cpp
+++ b/Engine/Plugins/MeshPainting/Source/MeshPaintingToolset/Private/MeshPaintHelpers.cpp
@@ -931,9 +931,10 @@ void UMeshPaintingSubsystem::ApplyFillWithMask(FColor& InOutColor, const FColor&
 
 void UMeshPaintingSubsystem::ForceRenderMeshLOD(UMeshComponent* Component, int32 LODIndex)
 {
+	// This seems dangerous. What if we save the actor while in forced LOD mode? What if we are stomping an art intended forced LOD?
 	if (UStaticMeshComponent* StaticMeshComponent = Cast<UStaticMeshComponent>(Component))
 	{
-		StaticMeshComponent->ForcedLodModel = LODIndex + 1;
+		StaticMeshComponent->SetForcedLodModel(LODIndex + 1);
 	}
 	else if (USkeletalMeshComponent* SkeletalMeshComponent = Cast<USkeletalMeshComponent>(Component))
 	{
