Fix potential nullptr deref when hammering undo on mesh texture painting.
Last undo will deselect the mesh, which in turn should disable paint tool. But there is one frame before disable is processed where the paint tool tick runs with no selected items. So we can't assume that the selected array is non-empy here.

diff --git a/Engine/Plugins/MeshPainting/Source/MeshPaintingToolset/Private/MeshTexturePaintingTool.cpp b/Engine/Plugins/MeshPainting/Source/MeshPaintingToolset/Private/MeshTexturePaintingTool.cpp
--- a/Engine/Plugins/MeshPainting/Source/MeshPaintingToolset/Private/MeshTexturePaintingTool.cpp
+++ b/Engine/Plugins/MeshPainting/Source/MeshPaintingToolset/Private/MeshTexturePaintingTool.cpp
@@ -162,6 +162,8 @@ void UMeshTexturePaintingTool::OnTick(float DeltaTime)
 	UMeshPaintingSubsystem* MeshPaintingSubsystem = GEngine->GetEngineSubsystem<UMeshPaintingSubsystem>();
 	if (MeshPaintingSubsystem)
 	{
+		TArray<UMeshComponent*> SelectedMeshComponents = MeshPaintingSubsystem->GetSelectedMeshComponents();
+
 		if (bRequestPaintBucketFill)
 		{
 			FMeshPaintParameters BucketFillParams;
@@ -181,7 +183,6 @@ void UMeshTexturePaintingTool::OnTick(float DeltaTime)
 				BucketFillParams.bUseFillBucket = true;
 			}
 
-			TArray<UMeshComponent*> SelectedMeshComponents = MeshPaintingSubsystem->GetSelectedMeshComponents();
 			for (int32 j = 0; j < SelectedMeshComponents.Num(); ++j)
 			{
 				UMeshComponent* SelectedComponent = SelectedMeshComponents[j];
@@ -241,8 +242,8 @@ void UMeshTexturePaintingTool::OnTick(float DeltaTime)
 			}
 		}
 
-		UMeshComponent* FirstSelectedComponent = MeshPaintingSubsystem->GetSelectedMeshComponents()[0];
-		if (MeshPaintingSubsystem->bNeedsRecache || (PaintableTextures.Num() > 0 && TextureProperties && GetSelectedPaintTexture(FirstSelectedComponent) == nullptr))
+		UMeshComponent* FirstSelectedComponent = SelectedMeshComponents.IsValidIndex(0) ? SelectedMeshComponents[0] : nullptr;
+		if (MeshPaintingSubsystem->bNeedsRecache || (PaintableTextures.Num() > 0 && GetSelectedPaintTexture(FirstSelectedComponent) == nullptr))
 		{
 			CacheSelectionData();
 			CacheTexturePaintData();
