Fix brief glitch when texture color mesh painting in a scene with only one mesh paint texture.
Issue was that if there is only one such texture, then the mesh paint virtual texture tracking loses it for a couple of frames when it texture compiles after a paint brush stroke.
That would lead to not having a page table registered even though we still need it to support the virtual texture adapter view that we are visualizing.
The fix is to create a dummy texture owned by the paint tool. That way we never drop to zero textures in the scene.

diff --git a/Engine/Plugins/MeshPainting/Source/MeshPaintingToolset/Private/MeshPaintHelpers.cpp b/Engine/Plugins/MeshPainting/Source/MeshPaintingToolset/Private/MeshPaintHelpers.cpp
--- a/Engine/Plugins/MeshPainting/Source/MeshPaintingToolset/Private/MeshPaintHelpers.cpp
+++ b/Engine/Plugins/MeshPainting/Source/MeshPaintingToolset/Private/MeshPaintHelpers.cpp
@@ -79,6 +79,18 @@ void UMeshPaintingSubsystem::RemoveComponentInstanceVertexColors(UStaticMeshComp
 	}
 }
 
+UTexture* UMeshPaintingSubsystem::CreateMeshPaintTexture(UObject* Outer,  uint32 TextureSize)
+{
+	const uint32 AlignedTextureSize = MeshPaintVirtualTexture::GetAlignedTextureSize(TextureSize);
+	const uint32 TextureNumMips = FMath::FloorLog2(AlignedTextureSize) + 1;
+
+	UMeshPaintVirtualTexture* NewTexture = NewObject<UMeshPaintVirtualTexture>(Outer);
+	NewTexture->Source.Init(AlignedTextureSize, AlignedTextureSize, 1, TextureNumMips, TSF_BGRA8);
+	NewTexture->UpdateResource();
+
+	return NewTexture;
+}
+
 void UMeshPaintingSubsystem::CreateComponentMeshPaintTexture(UStaticMeshComponent* StaticMeshComponent)
 {
 	if (StaticMeshComponent != nullptr && StaticMeshComponent->GetMeshPaintTexture() == nullptr && StaticMeshComponent->CanMeshPaintTextureColors())
diff --git a/Engine/Plugins/MeshPainting/Source/MeshPaintingToolset/Private/MeshTexturePaintingTool.cpp b/Engine/Plugins/MeshPainting/Source/MeshPaintingToolset/Private/MeshTexturePaintingTool.cpp
--- a/Engine/Plugins/MeshPainting/Source/MeshPaintingToolset/Private/MeshTexturePaintingTool.cpp
+++ b/Engine/Plugins/MeshPainting/Source/MeshPaintingToolset/Private/MeshTexturePaintingTool.cpp
@@ -655,10 +655,7 @@ void UMeshTexturePaintingTool::FinishPainting()
 
 FPaintTexture2DData* UMeshTexturePaintingTool::GetPaintTargetData(const UTexture2D* InTexture)
 {
-	checkf(InTexture != nullptr, TEXT("Invalid Texture ptr"));
-	/** Retrieve target paint data for the given texture */
-	FPaintTexture2DData* TextureData = PaintTargetData.Find(InTexture);
-	return TextureData;
+	return PaintTargetData.Find(InTexture);
 }
 
 FPaintTexture2DData* UMeshTexturePaintingTool::AddPaintTargetData(UTexture2D* InTexture)
@@ -1361,6 +1358,20 @@ void UMeshTextureColorPaintingTool::Setup()
 	Super::Setup();
 	ColorProperties = Cast<UMeshTextureColorPaintingToolProperties>(BrushProperties);
 
+	if (UMeshPaintingSubsystem* MeshPaintingSubsystem = GEngine->GetEngineSubsystem<UMeshPaintingSubsystem>())
+	{
+		// Create a dummy mesh paint virtual texture for the lifetime of the paint tool.
+		// This keeps at least one virtual texture alive during painting.
+		// Otherwise, if there is only one "real" virtual texture in the scene and we paint on it, 
+		// it will be deallocted for one or two frames during texture compilation after each paint stroke.
+		// For those frames there would be _no_ remainging allocated VTs to use for the FMeshPaintVirtualTextureSceneExtension
+		// and which would leave no page table bound for sampling the virtual texture adaptor that wraps the 
+		// painting render target. That would result in a flicker where the lack of page table means the mesh paint
+		// virtual texture gets its fallback color when sampling.
+		// Holding this dummy texture prevents that from happening.
+		MeshPaintDummyTexture = MeshPaintingSubsystem->CreateMeshPaintTexture(this, 1);
+	}
+
 	GetToolManager()->DisplayMessage(
 		LOCTEXT("OnStartTextureColorPaintTool", "Paint colors to the Mesh Paint Texture object stored on mesh components."),
 		EToolMessageLevel::UserNotification);
diff --git a/Engine/Plugins/MeshPainting/Source/MeshPaintingToolset/Public/MeshPaintHelpers.h b/Engine/Plugins/MeshPainting/Source/MeshPaintingToolset/Public/MeshPaintHelpers.h
--- a/Engine/Plugins/MeshPainting/Source/MeshPaintingToolset/Public/MeshPaintHelpers.h
+++ b/Engine/Plugins/MeshPainting/Source/MeshPaintingToolset/Public/MeshPaintHelpers.h
@@ -137,6 +137,8 @@ public:
 	/** Removes vertex colors associated with the mesh component */
 	void RemoveComponentInstanceVertexColors(UStaticMeshComponent* StaticMeshComponent);
 
+	/** Creates and returns a mesh paint texture that isn't attached to a mesh component */
+	UTexture* CreateMeshPaintTexture(UObject* Outer, uint32 TextureSize);
 	/** Creates mesh paint texture associated with the mesh component */
 	void CreateComponentMeshPaintTexture(UStaticMeshComponent* StaticMeshComponent);
 	void CreateComponentMeshPaintTexture(UStaticMeshComponent* StaticMeshComponent, FImageView const& InImage);
diff --git a/Engine/Plugins/MeshPainting/Source/MeshPaintingToolset/Public/MeshTexturePaintingTool.h b/Engine/Plugins/MeshPainting/Source/MeshPaintingToolset/Public/MeshTexturePaintingTool.h
--- a/Engine/Plugins/MeshPainting/Source/MeshPaintingToolset/Public/MeshTexturePaintingTool.h
+++ b/Engine/Plugins/MeshPainting/Source/MeshPaintingToolset/Public/MeshTexturePaintingTool.h
@@ -264,6 +264,9 @@ protected:
 private:
 	UPROPERTY(Transient)
 	TObjectPtr<UMeshTextureColorPaintingToolProperties> ColorProperties;
+
+	UPROPERTY(Transient)
+	TObjectPtr<UTexture> MeshPaintDummyTexture;
 };
 
 /**
