Fix issue in mesh paint mode where after auto save we will have no tool selected.

diff --git a/Engine/Plugins/MeshPainting/Source/MeshPaintEditorMode/Private/MeshPaintMode.cpp b/Engine/Plugins/MeshPainting/Source/MeshPaintEditorMode/Private/MeshPaintMode.cpp
--- a/Engine/Plugins/MeshPainting/Source/MeshPaintEditorMode/Private/MeshPaintMode.cpp
+++ b/Engine/Plugins/MeshPainting/Source/MeshPaintEditorMode/Private/MeshPaintMode.cpp
@@ -250,7 +250,12 @@ void UMeshPaintMode::Tick(FEditorViewportClient* ViewportClient, float DeltaTime
 		TargetTab = MeshPaintMode_TextureAsset;
 		CurrentActiveMode = EMeshPaintActiveMode::Texture;
 	}
-	if ( TargetTab != ActiveTab)
+	else if (ActiveTool == nullptr)
+	{
+		TargetTab = MeshPaintMode_VertexColor;
+		CurrentActiveMode = EMeshPaintActiveMode::VertexColor;
+	}
+	if (TargetTab != ActiveTab)
 	{
 		Toolkit->SetCurrentPalette(TargetTab);
 	}
