Fix mesh paint seam painting bug.
When collecting the seam mask we were assuming that the FStaticMeshSection::MaterialIndex was the same as the section index.
When that wasn't true we were using a seam mask from the wrong material section.
Also tidied/optimized the code there so that we iterate over sections and their triangles instead of needing to gather all of the triangles upfront.
Also fixed bug where we were generating seam mask every paint stroke instead of once at start of painting.

diff --git a/Engine/Plugins/MeshPainting/Source/MeshPaintingToolset/Private/TexturePaintToolset.cpp b/Engine/Plugins/MeshPainting/Source/MeshPaintingToolset/Private/TexturePaintToolset.cpp
--- a/Engine/Plugins/MeshPainting/Source/MeshPaintingToolset/Private/TexturePaintToolset.cpp
+++ b/Engine/Plugins/MeshPainting/Source/MeshPaintingToolset/Private/TexturePaintToolset.cpp
@@ -135,11 +135,11 @@ bool UTexturePaintToolset::GenerateSeamMask(UMeshComponent* MeshComponent, int32
 	FTextureRenderTargetResource* RenderTargetResource = SeamRenderTexture->GameThread_GetRenderTargetResource();
 	check(RenderTargetResource != nullptr);
 
-	int32 NumElements = StaticMeshComponent->GetNumMaterials();
+	const int32 NumElements = StaticMeshComponent->GetNumMaterials();
 	UTexture2D* TargetTexture2D = Texture;
 
-	// Store info that tells us if the element material uses our target texture so we don't have to do a usestexture() call for each tri.  We will
-	// use this info to eliminate triangles that do not use our texture.
+	// Store info that tells us if the element material uses our target texture.
+	// We will use this info to eliminate triangles that do not use our texture.
 	TArray< bool > ElementUsesTargetTexture;
 	ElementUsesTargetTexture.AddZeroed(NumElements);
 	for (int32 ElementIndex = 0; ElementIndex < NumElements; ElementIndex++)
@@ -160,144 +160,101 @@ bool UTexturePaintToolset::GenerateSeamMask(UMeshComponent* MeshComponent, int32
 		}
 	}
 
-	// Make sure we're dealing with triangle lists
-	FIndexArrayView Indices = LODModel.IndexBuffer.GetArrayView();
-	const uint32 NumIndexBufferIndices = Indices.Num();
-	check(NumIndexBufferIndices % 3 == 0);
-	const uint32 NumTriangles = NumIndexBufferIndices / 3;
-
-	static TArray< int32 > InfluencedTriangles;
-	InfluencedTriangles.Empty(NumTriangles);
-
-	// For each triangle in the mesh
-	for (uint32 TriIndex = 0; TriIndex < NumTriangles; ++TriIndex)
-	{
-		// At least one triangle vertex was influenced.
-		bool bAddTri = false;
-
-		// Check to see if the sub-element that this triangle belongs to actually uses our paint target texture in its material
-		for (int32 ElementIndex = 0; ElementIndex < NumElements; ElementIndex++)
-		{
-			//FStaticMeshElement& Element = LODModel.Elements[ ElementIndex ];
-			FStaticMeshSection& Element = LODModel.Sections[ElementIndex];
-
-
-			if ((TriIndex >= Element.FirstIndex / 3) &&
-				(TriIndex < Element.FirstIndex / 3 + Element.NumTriangles))
-			{
-
-				// The triangle belongs to this element, now we need to check to see if the element material uses our target texture.
-				if (TargetTexture2D != nullptr && ElementUsesTargetTexture[ElementIndex] == true)
-				{
-					bAddTri = true;
-				}
-
-				// Triangles can only be part of one element so we do not need to continue to other elements.
-				break;
-			}
-
-		}
-
-		if (bAddTri)
-		{
-			InfluencedTriangles.Add(TriIndex);
-		}
-
-	}
-
 	{
 		// Create a canvas for the render target and clear it to white
 		FCanvas Canvas(RenderTargetResource, nullptr, FGameTime(), GEditor->GetEditorWorldContext().World()->GetFeatureLevel());
 		Canvas.Clear(FLinearColor::White);
 
+		FIndexArrayView Indices = LODModel.IndexBuffer.GetArrayView();
+		
 		TArray<FCanvasUVTri> TriList;
-		FCanvasUVTri EachTri;
-		EachTri.V0_Color = FLinearColor::Black;
-		EachTri.V1_Color = FLinearColor::Black;
-		EachTri.V2_Color = FLinearColor::Black;
-
-		for (int32 CurIndex = 0; CurIndex < InfluencedTriangles.Num(); ++CurIndex)
+		for (int32 ElementIndex = 0; ElementIndex < NumElements; ++ElementIndex)
 		{
-			const int32 TriIndex = InfluencedTriangles[CurIndex];
+			FStaticMeshSection& Element = LODModel.Sections[ElementIndex];
 
-			// Grab the vertex indices and points for this triangle
-			FVector2D TriUVs[3];
-			FVector2D UVMin(99999.9f, 99999.9f);
-			FVector2D UVMax(-99999.9f, -99999.9f);
-			for (int32 TriVertexNum = 0; TriVertexNum < 3; ++TriVertexNum)
+			if (ElementUsesTargetTexture[Element.MaterialIndex])
 			{
-				const int32 VertexIndex = Indices[TriIndex * 3 + TriVertexNum];
-				TriUVs[TriVertexNum] = FVector2D(LODModel.VertexBuffers.StaticMeshVertexBuffer.GetVertexUV(VertexIndex, UVSet));
-
-				// Update bounds
-				float U = TriUVs[TriVertexNum].X;
-				float V = TriUVs[TriVertexNum].Y;
-
-				if (U < UVMin.X)
-				{
-					UVMin.X = U;
-				}
-				if (U > UVMax.X)
+				for (uint32 TriIndex = Element.FirstIndex / 3u; TriIndex < Element.FirstIndex / 3u + Element.NumTriangles; ++TriIndex)
 				{
-					UVMax.X = U;
+					// Grab the vertex indices and points for this triangle
+					FVector2D TriUVs[3];
+					FVector2D UVMin(99999.9f, 99999.9f);
+					FVector2D UVMax(-99999.9f, -99999.9f);
+					for (int32 TriVertexNum = 0; TriVertexNum < 3; ++TriVertexNum)
+					{
+						const int32 VertexIndex = Indices[TriIndex * 3 + TriVertexNum];
+						TriUVs[TriVertexNum] = FVector2D(LODModel.VertexBuffers.StaticMeshVertexBuffer.GetVertexUV(VertexIndex, UVSet));
+
+						// Update bounds
+						float U = TriUVs[TriVertexNum].X;
+						float V = TriUVs[TriVertexNum].Y;
+
+						if (U < UVMin.X)
+						{
+							UVMin.X = U;
+						}
+						if (U > UVMax.X)
+						{
+							UVMax.X = U;
+						}
+						if (V < UVMin.Y)
+						{
+							UVMin.Y = V;
+						}
+						if (V > UVMax.Y)
+						{
+							UVMax.Y = V;
+						}
+					}
+
+					// If the triangle lies entirely outside of the 0.0-1.0 range, we'll transpose it back
+					FVector2D UVOffset(0.0f, 0.0f);
+					if (UVMax.X > 1.0f)
+					{
+						UVOffset.X = -FMath::FloorToInt(UVMin.X);
+					}
+					else if (UVMin.X < 0.0f)
+					{
+						UVOffset.X = 1.0f + FMath::FloorToInt(-UVMax.X);
+					}
+
+					if (UVMax.Y > 1.0f)
+					{
+						UVOffset.Y = -FMath::FloorToInt(UVMin.Y);
+					}
+					else if (UVMin.Y < 0.0f)
+					{
+						UVOffset.Y = 1.0f + FMath::FloorToInt(-UVMax.Y);
+					}
+
+					// Note that we "wrap" the texture coordinates here to handle the case where the user
+					// is painting on a tiling texture, or with the UVs out of bounds.  Ideally all of the
+					// UVs would be in the 0.0 - 1.0 range but sometimes content isn't setup that way.
+					// @todo MeshPaint: Handle triangles that cross the 0.0-1.0 UV boundary?
+					FVector2D TrianglePoints[3];
+					for (int32 TriVertexNum = 0; TriVertexNum < 3; ++TriVertexNum)
+					{
+						TriUVs[TriVertexNum].X += UVOffset.X;
+						TriUVs[TriVertexNum].Y += UVOffset.Y;
+
+						TrianglePoints[TriVertexNum].X = TriUVs[TriVertexNum].X * Width;
+						TrianglePoints[TriVertexNum].Y = TriUVs[TriVertexNum].Y * Height;
+					}
+
+					FCanvasUVTri& Tri = TriList.Emplace_GetRef();
+					Tri.V0_Pos = TrianglePoints[0];
+					Tri.V0_UV = TriUVs[0];
+					Tri.V0_Color = FLinearColor::Black;
+					Tri.V1_Pos = TrianglePoints[1];
+					Tri.V1_UV = TriUVs[1];
+					Tri.V1_Color = FLinearColor::Black;
+					Tri.V2_Pos = TrianglePoints[2];
+					Tri.V2_UV = TriUVs[2];
+					Tri.V2_Color = FLinearColor::Black;
 				}
-				if (V < UVMin.Y)
-				{
-					UVMin.Y = V;
-				}
-				if (V > UVMax.Y)
-				{
-					UVMax.Y = V;
-				}
-
-			}
-
-			// If the triangle lies entirely outside of the 0.0-1.0 range, we'll transpose it back
-			FVector2D UVOffset(0.0f, 0.0f);
-			if (UVMax.X > 1.0f)
-			{
-				UVOffset.X = -FMath::FloorToInt(UVMin.X);
-			}
-			else if (UVMin.X < 0.0f)
-			{
-				UVOffset.X = 1.0f + FMath::FloorToInt(-UVMax.X);
-			}
-
-			if (UVMax.Y > 1.0f)
-			{
-				UVOffset.Y = -FMath::FloorToInt(UVMin.Y);
-			}
-			else if (UVMin.Y < 0.0f)
-			{
-				UVOffset.Y = 1.0f + FMath::FloorToInt(-UVMax.Y);
 			}
-
-
-			// Note that we "wrap" the texture coordinates here to handle the case where the user
-			// is painting on a tiling texture, or with the UVs out of bounds.  Ideally all of the
-			// UVs would be in the 0.0 - 1.0 range but sometimes content isn't setup that way.
-			// @todo MeshPaint: Handle triangles that cross the 0.0-1.0 UV boundary?
-			FVector2D TrianglePoints[3];
-			for (int32 TriVertexNum = 0; TriVertexNum < 3; ++TriVertexNum)
-			{
-				TriUVs[TriVertexNum].X += UVOffset.X;
-				TriUVs[TriVertexNum].Y += UVOffset.Y;
-
-				TrianglePoints[TriVertexNum].X = TriUVs[TriVertexNum].X * Width;
-				TrianglePoints[TriVertexNum].Y = TriUVs[TriVertexNum].Y * Height;
-			}
-
-			EachTri.V0_Pos = TrianglePoints[0];
-			EachTri.V0_UV = TriUVs[0];
-			EachTri.V0_Color = FLinearColor::Black;
-			EachTri.V1_Pos = TrianglePoints[1];
-			EachTri.V1_UV = TriUVs[1];
-			EachTri.V1_Color = FLinearColor::Black;
-			EachTri.V2_Pos = TrianglePoints[2];
-			EachTri.V2_UV = TriUVs[2];
-			EachTri.V2_Color = FLinearColor::Black;
-			TriList.Add(EachTri);
 		}
+
 		if (TriList.Num())
 		{
 			// Setup the tri render item with the list of tris
@@ -310,14 +267,11 @@ bool UTexturePaintToolset::GenerateSeamMask(UMeshComponent* MeshComponent, int32
 		}
 	}
 
-
+	ENQUEUE_RENDER_COMMAND(UpdateMeshPaintRTCommand5)(
+		[RenderTargetResource](FRHICommandListImmediate& RHICmdList)
 	{
-		ENQUEUE_RENDER_COMMAND(UpdateMeshPaintRTCommand5)(
-			[RenderTargetResource](FRHICommandListImmediate& RHICmdList)
-		{
-			TransitionAndCopyTexture(RHICmdList, RenderTargetResource->GetRenderTargetTexture(), RenderTargetResource->TextureRHI, {});
-		});
-	}
+		TransitionAndCopyTexture(RHICmdList, RenderTargetResource->GetRenderTargetTexture(), RenderTargetResource->TextureRHI, {});
+	});
 
 	return RetVal;
 }
