Change mesh texture color paint enabled logic.
Previously we required that the material sampled mesh paint textures before allowing any painting.
This is difficult to discover though. So now always allow.
We already added a flag on the mesh asset and component to disallow texture color painting, so there is an easy to understand mechanism for users to restrict painting without relying on what the material does (which is a bit opaque and different to the way vertex painting functions).

diff --git a/Engine/Plugins/MeshPainting/Source/MeshPaintingToolset/Private/IMeshPaintComponentAdapter.cpp b/Engine/Plugins/MeshPainting/Source/MeshPaintingToolset/Private/IMeshPaintComponentAdapter.cpp
--- a/Engine/Plugins/MeshPainting/Source/MeshPaintingToolset/Private/IMeshPaintComponentAdapter.cpp
+++ b/Engine/Plugins/MeshPainting/Source/MeshPaintingToolset/Private/IMeshPaintComponentAdapter.cpp
@@ -8,7 +8,6 @@
 #include "Engine/World.h"
 #include "MaterialShared.h"
 #include "Materials/Material.h"
-#include "Materials/MaterialExpressionMeshPaintTextureObject.h"
 #include "Materials/MaterialExpressionTextureCoordinate.h"
 #include "Materials/MaterialExpressionTextureSampleParameter.h"
 #include "Materials/MaterialInstance.h"
@@ -110,16 +109,6 @@ void IMeshPaintComponentAdapter::DefaultQueryPaintableTextures(int32 MaterialInd
 						OutDefaultIndex = TextureIndex;
 					}
 				}
-
-				// If material samples the primitive mesh paint texture, then add it here.
-				if (UMaterialExpressionMeshPaintTextureObject* MeshPaintTextureExpression = Cast<UMaterialExpressionMeshPaintTextureObject>(Expression))
-				{
-					if (UTexture* MeshPaintTexture = MeshComponent->GetMeshPaintTexture())
-					{
-						const int32 CoordinateIndex = MeshComponent->GetMeshPaintTextureCoordinateIndex();
-						InOutTextureList.AddUnique(FPaintableTexture(MeshPaintTexture, CoordinateIndex, true));
-					}
-				}
 			}
 		}
 		// Make sure to include all texture parameters, this will include all of the texture parameters from internal material functions
@@ -151,6 +140,13 @@ void IMeshPaintComponentAdapter::DefaultQueryPaintableTextures(int32 MaterialInd
 			break;
 		}
 	}
+
+	// If the component has a mesh paint texture, then add it here.
+	if (UTexture* MeshPaintTexture = MeshComponent->GetMeshPaintTexture())
+	{
+		const int32 CoordinateIndex = MeshComponent->GetMeshPaintTextureCoordinateIndex();
+		InOutTextureList.AddUnique(FPaintableTexture(MeshPaintTexture, CoordinateIndex, true));
+	}
 }
 
 namespace UE::MeshPaintingToolset
