Mesh paint mode: fix seam expansion for texture color painting.
Previously no material sections were chosen to mask seam expansion because we only used logic from texture asset painting.
Added different logic for texture color painting, and select between the two paths now.

diff --git a/Engine/Plugins/MeshPainting/Source/MeshPaintingToolset/Private/TexturePaintToolset.cpp b/Engine/Plugins/MeshPainting/Source/MeshPaintingToolset/Private/TexturePaintToolset.cpp
--- a/Engine/Plugins/MeshPainting/Source/MeshPaintingToolset/Private/TexturePaintToolset.cpp
+++ b/Engine/Plugins/MeshPainting/Source/MeshPaintingToolset/Private/TexturePaintToolset.cpp
@@ -135,6 +135,8 @@ bool UTexturePaintToolset::GenerateSeamMask(UMeshComponent* MeshComponent, int32
 	FTextureRenderTargetResource* RenderTargetResource = SeamRenderTexture->GameThread_GetRenderTargetResource();
 	check(RenderTargetResource != nullptr);
 
+	const bool bIsMeshPaintTexture = MeshComponent->GetMeshPaintTexture() == Texture;
+
 	const int32 NumElements = StaticMeshComponent->GetNumMaterials();
 	UTexture2D* TargetTexture2D = Texture;
 
@@ -149,13 +151,20 @@ bool UTexturePaintToolset::GenerateSeamMask(UMeshComponent* MeshComponent, int32
 		UMaterialInterface* ElementMat = StaticMeshComponent->GetMaterial(ElementIndex);
 		if (ElementMat != nullptr)
 		{
-			ElementUsesTargetTexture[ElementIndex] |= DoesMaterialUseTexture(ElementMat, TargetTexture2D);
-
-			if (ElementUsesTargetTexture[ElementIndex] == false && RenderTargetTexture != nullptr)
+			if (bIsMeshPaintTexture)
+			{
+				ElementUsesTargetTexture[ElementIndex] |= ElementMat->HasMeshPaintTexture();
+			}
+			else
 			{
-				// If we didn't get a match on our selected texture, we'll check to see if the the material uses a
-				//  render target texture override that we put on during painting.
-				ElementUsesTargetTexture[ElementIndex] |= DoesMaterialUseTexture(ElementMat, RenderTargetTexture);
+				ElementUsesTargetTexture[ElementIndex] |= DoesMaterialUseTexture(ElementMat, TargetTexture2D);
+
+				if (ElementUsesTargetTexture[ElementIndex] == false && RenderTargetTexture != nullptr)
+				{
+					// If we didn't get a match on our selected texture, we'll check to see if the the material uses a
+					//  render target texture override that we put on during painting.
+					ElementUsesTargetTexture[ElementIndex] |= DoesMaterialUseTexture(ElementMat, RenderTargetTexture);
+				}
 			}
 		}
 	}
@@ -372,13 +381,23 @@ void UTexturePaintToolset::FindMaterialIndicesUsingTexture(const UTexture* Textu
 {
 	checkf(Texture && MeshComponent, TEXT("Invalid Texture of MeshComponent"));
 
+	const bool bIsMeshPaintTexture = MeshComponent->GetMeshPaintTexture() == Texture;
+
 	const int32 NumMaterials = MeshComponent->GetNumMaterials();
 	for (int32 MaterialIndex = 0; MaterialIndex < NumMaterials; ++MaterialIndex)
 	{
 		const UMaterialInterface* MaterialInterface = MeshComponent->GetMaterial(MaterialIndex);
 		if (MaterialInterface)
 		{
-			const bool bUsesTexture = DoesMaterialUseTexture(MaterialInterface, Texture);
+			bool bUsesTexture = false;
+			if (bIsMeshPaintTexture)
+			{
+				bUsesTexture = MaterialInterface->HasMeshPaintTexture();
+			}
+			else
+			{
+				bUsesTexture = DoesMaterialUseTexture(MaterialInterface, Texture);
+			}
 			if (bUsesTexture)
 			{
 				OutIndices.AddUnique(MaterialIndex);
