Disable UI elements of the texture color mesh paint tool when mesh paint textures are disabled for the project.
Do this by taking the support state into account for UStaticMeshComponent::CanMeshPaintTextureColors() and adding a check for that where necessary.

diff --git a/Engine/Plugins/MeshPainting/Source/MeshPaintEditorMode/Private/MeshPaintMode.cpp b/Engine/Plugins/MeshPainting/Source/MeshPaintEditorMode/Private/MeshPaintMode.cpp
--- a/Engine/Plugins/MeshPainting/Source/MeshPaintEditorMode/Private/MeshPaintMode.cpp
+++ b/Engine/Plugins/MeshPainting/Source/MeshPaintEditorMode/Private/MeshPaintMode.cpp
@@ -1124,7 +1124,7 @@ void UMeshPaintMode::CopyMeshPaintTexture()
 bool UMeshPaintMode::CanCopyMeshPaintTexture() const
 {
 	const TArray<UStaticMeshComponent*> StaticMeshComponents = GetSelectedComponents<UStaticMeshComponent>();
-	if (StaticMeshComponents.Num() == 1 && StaticMeshComponents[0]->GetMeshPaintTexture() != nullptr)
+	if (StaticMeshComponents.Num() == 1 && StaticMeshComponents[0]->GetMeshPaintTexture() != nullptr && StaticMeshComponents[0]->CanMeshPaintTextureColors())
 	{
 		return true;
 	}
@@ -1176,17 +1176,20 @@ void UMeshPaintMode::PasteMeshPaintTexture()
 
 bool UMeshPaintMode::CanPasteMeshPaintTexture() const
 {
-	const TArray<UStaticMeshComponent*> StaticMeshComponents = GetSelectedComponents<UStaticMeshComponent>();
-	if (StaticMeshComponents.Num() == 0)
-	{
-		return false;
-	}
 	FImage const& Image = GEngine->GetEngineSubsystem<UMeshPaintingSubsystem>()->GetCopiedTexture();
 	if (Image.GetNumPixels() == 0)
 	{
 		return false;
 	}
-	return true;
+	const TArray<UStaticMeshComponent*> StaticMeshComponents = GetSelectedComponents<UStaticMeshComponent>();
+	for (UStaticMeshComponent* Component : StaticMeshComponents)
+	{
+		if (Component->CanMeshPaintTextureColors())
+		{
+			return true;
+		}
+	}
+	return false;
 }
 
 void UMeshPaintMode::Paste()
@@ -1264,8 +1267,15 @@ void UMeshPaintMode::ImportMeshPaintTextureFromVertexColors()
 
 bool UMeshPaintMode::CanImportMeshPaintTextureFromVertexColors() const
 {
-	const TArray<UStaticMeshComponent*> MeshComponents = GetSelectedComponents<UStaticMeshComponent>();
-	return MeshComponents.Num() > 0;
+	const TArray<UStaticMeshComponent*> StaticMeshComponents = GetSelectedComponents<UStaticMeshComponent>();
+	for (UStaticMeshComponent* Component : StaticMeshComponents)
+	{
+		if (Component->CanMeshPaintTextureColors())
+		{
+			return true;
+		}
+	}
+	return false;
 }
 
 void UMeshPaintMode::FixTextureColors()
diff --git a/Engine/Plugins/MeshPainting/Source/MeshPaintEditorMode/Private/MeshPaintModeHelpers.cpp b/Engine/Plugins/MeshPainting/Source/MeshPaintEditorMode/Private/MeshPaintModeHelpers.cpp
--- a/Engine/Plugins/MeshPainting/Source/MeshPaintEditorMode/Private/MeshPaintModeHelpers.cpp
+++ b/Engine/Plugins/MeshPainting/Source/MeshPaintEditorMode/Private/MeshPaintModeHelpers.cpp
@@ -928,9 +928,12 @@ bool UMeshPaintModeSubsystem::CanFixTextureColors(const TArray<UMeshComponent*>&
 		{
 			if (UTexture* Texture = StaticMeshComponent->GetMeshPaintTexture())
 			{
-				if (StaticMeshComponent->GetMeshPaintTextureResolution() != Texture->Source.GetSizeX())
+				if (StaticMeshComponent->CanMeshPaintTextureColors())
 				{
-					return true;
+					if (StaticMeshComponent->GetMeshPaintTextureResolution() != Texture->Source.GetSizeX())
+					{
+						return true;
+					}
 				}
 			}
 		}
diff --git a/Engine/Source/Runtime/Engine/Private/Components/StaticMeshComponent.cpp b/Engine/Source/Runtime/Engine/Private/Components/StaticMeshComponent.cpp
--- a/Engine/Source/Runtime/Engine/Private/Components/StaticMeshComponent.cpp
+++ b/Engine/Source/Runtime/Engine/Private/Components/StaticMeshComponent.cpp
@@ -3250,6 +3250,13 @@ bool UStaticMeshComponent::CanMeshPaintTextureColors() const
 	{
 		return false;
 	}
+	if (FSceneInterface* Scene = GetScene())
+	{
+		if (!MeshPaintVirtualTexture::IsSupported(Scene->GetShaderPlatform()))
+		{
+			return false;
+		}
+	}
 	return true;
 }
 
