Mesh paint. Fix edge case for texture seam painting where the painted texture is unused by any material section.
In that case we don't want to blur all painting, so render all material sections into the seam mask instead of none.

diff --git a/Engine/Plugins/MeshPainting/Source/MeshPaintingToolset/Private/TexturePaintToolset.cpp b/Engine/Plugins/MeshPainting/Source/MeshPaintingToolset/Private/TexturePaintToolset.cpp
--- a/Engine/Plugins/MeshPainting/Source/MeshPaintingToolset/Private/TexturePaintToolset.cpp
+++ b/Engine/Plugins/MeshPainting/Source/MeshPaintingToolset/Private/TexturePaintToolset.cpp
@@ -142,7 +142,8 @@ bool UTexturePaintToolset::GenerateSeamMask(UMeshComponent* MeshComponent, int32
 
 	// Store info that tells us if the element material uses our target texture.
 	// We will use this info to eliminate triangles that do not use our texture.
-	TArray< bool > ElementUsesTargetTexture;
+	TArray<bool> ElementUsesTargetTexture;
+	bool bAnyElementUsesTargetTexture = false;
 	ElementUsesTargetTexture.AddZeroed(NumElements);
 	for (int32 ElementIndex = 0; ElementIndex < NumElements; ElementIndex++)
 	{
@@ -167,6 +168,11 @@ bool UTexturePaintToolset::GenerateSeamMask(UMeshComponent* MeshComponent, int32
 				}
 			}
 		}
+
+		// We track if there is no section that uses the texture.
+		// That would be a special case where we are painting without any context for the seams. Then seam painting would expand/blur all painting.
+		// To avoid that it's better to render _all_ sections into the seam mask rather than none.
+		bAnyElementUsesTargetTexture |= ElementUsesTargetTexture[ElementIndex];
 	}
 
 	{
@@ -181,7 +187,7 @@ bool UTexturePaintToolset::GenerateSeamMask(UMeshComponent* MeshComponent, int32
 		{
 			FStaticMeshSection& Element = LODModel.Sections[ElementIndex];
 
-			if (ElementUsesTargetTexture[Element.MaterialIndex])
+			if (ElementUsesTargetTexture[Element.MaterialIndex] || !bAnyElementUsesTargetTexture)
 			{
 				for (uint32 TriIndex = Element.FirstIndex / 3u; TriIndex < Element.FirstIndex / 3u + Element.NumTriangles; ++TriIndex)
 				{
