Mesh Paint Mode UI tweaks.
Move Paint and Erase Color to the Color Painting Category from the Brush Category.
Add clamping to VertexPreviewSize.

diff --git a/Engine/Plugins/MeshPainting/Source/MeshPaintEditorMode/Private/BrushSettingsCustomization.cpp b/Engine/Plugins/MeshPainting/Source/MeshPaintEditorMode/Private/BrushSettingsCustomization.cpp
--- a/Engine/Plugins/MeshPainting/Source/MeshPaintEditorMode/Private/BrushSettingsCustomization.cpp
+++ b/Engine/Plugins/MeshPainting/Source/MeshPaintEditorMode/Private/BrushSettingsCustomization.cpp
@@ -48,6 +48,7 @@ TSharedRef<IDetailCustomization> FMeshPaintModeSettingsCustomization::MakeInstan
 void FMeshPaintModeSettingsCustomization::CustomizeDetails(IDetailLayoutBuilder& DetailLayout)
 {
 	IDetailCategoryBuilder& ResourceCategory = DetailLayout.EditCategory(TEXT("ResourceUsage"));
+	ResourceCategory.SetSortOrder(0);
 
 	ResourceCategory.AddCustomRow(NSLOCTEXT("VertexPaintSettings", "VertexColorSizeRow", "Vertex Color Size"))
 		.NameContent()
@@ -94,7 +95,7 @@ TSharedRef<IDetailCustomization> FMeshPaintingSettingsCustomization::MakeInstanc
 
 void FMeshPaintingSettingsCustomization::CustomizeDetails(IDetailLayoutBuilder& DetailLayout)
 {
-	IDetailCategoryBuilder& BrushCategory = DetailLayout.EditCategory(TEXT("Brush"), FText::GetEmpty(), ECategoryPriority::Important);
+	IDetailCategoryBuilder& ColorCategory = DetailLayout.EditCategory(TEXT("ColorPainting"));
 
 	// Customize paint color with a swap button
 	TSharedRef<IPropertyHandle> PaintColor = DetailLayout.GetProperty("PaintColor", UMeshPaintingToolProperties::StaticClass());
@@ -106,7 +107,7 @@ void FMeshPaintingSettingsCustomization::CustomizeDetails(IDetailLayoutBuilder&
 		TSharedPtr<SWidget> NameWidget;
 		TSharedPtr<SWidget> ValueWidget;
 
-		IDetailPropertyRow& PaintColorProp = BrushCategory.AddProperty(PaintColor);
+		IDetailPropertyRow& PaintColorProp = ColorCategory.AddProperty(PaintColor);
 		PaintColorProp.GetDefaultWidgets(NameWidget, ValueWidget, false);
 		FDetailWidgetRow& PaintRow = PaintColorProp.CustomWidget(true);
 		PaintRow.NameContent()
@@ -147,7 +148,7 @@ void FMeshPaintingSettingsCustomization::CustomizeDetails(IDetailLayoutBuilder&
 		];
 	}
 	{
-		IDetailPropertyRow& EraseColorProp = BrushCategory.AddProperty(EraseColor);
+		IDetailPropertyRow& EraseColorProp = ColorCategory.AddProperty(EraseColor);
 
 		TSharedPtr<SWidget> NameWidget;
 		TSharedPtr<SWidget> ValueWidget;
@@ -181,6 +182,7 @@ TSharedRef<IDetailCustomization> FVertexPaintingSettingsCustomization::MakeInsta
 void FVertexPaintingSettingsCustomization::CustomizeDetails(IDetailLayoutBuilder& DetailLayout)
 {
 	IDetailCategoryBuilder& VertexCategory = DetailLayout.EditCategory(TEXT("VertexPainting"));
+	VertexCategory.SetSortOrder(3);
 
 	/** Add custom row for painting on specific LOD level with callbacks to the painter to update the data */
 	TSharedRef<IPropertyHandle> LODPaintingEnabled = DetailLayout.GetProperty("bPaintOnSpecificLOD", UMeshVertexPaintingToolProperties::StaticClass());
@@ -303,7 +305,9 @@ void FVertexColorPaintingSettingsCustomization::CustomizeDetails(IDetailLayoutBu
 	FVertexPaintingSettingsCustomization::CustomizeDetails(DetailLayout);
 
 	IDetailCategoryBuilder& ColorCategory = DetailLayout.EditCategory(TEXT("ColorPainting"));
-	ColorCategory.SetSortOrder(0);
+	ColorCategory.SetSortOrder(1);
+	IDetailCategoryBuilder& BrushCategory = DetailLayout.EditCategory(TEXT("Brush"));
+	BrushCategory.SetSortOrder(2);
 
 	/** Creates a custom widget row containing all color channel flags */
 	TSharedRef<IPropertyHandle> RedChannel = DetailLayout.GetProperty("bWriteRed", UMeshVertexColorPaintingToolProperties::StaticClass());
@@ -352,16 +356,17 @@ void FVertexWeightPaintingSettingsCustomization::CustomizeDetails(IDetailLayoutB
 {
 	FVertexPaintingSettingsCustomization::CustomizeDetails(DetailLayout);
 
+	IDetailCategoryBuilder& WeightCategory = DetailLayout.EditCategory(TEXT("WeightPainting"));
+	WeightCategory.SetSortOrder(1);
+	IDetailCategoryBuilder& BrushCategory = DetailLayout.EditCategory(TEXT("Brush"));
+	BrushCategory.SetSortOrder(2);
+
 	// Brush colors not needed for weight painting.
-	IDetailCategoryBuilder& BrushCategory = DetailLayout.EditCategory(TEXT("Brush"), FText::GetEmpty(), ECategoryPriority::Important);
 	TSharedRef<IPropertyHandle> PaintColor = DetailLayout.GetProperty("PaintColor", UMeshPaintingToolProperties::StaticClass());
 	PaintColor->MarkHiddenByCustomization();
 	TSharedRef<IPropertyHandle> EraseColor = DetailLayout.GetProperty("EraseColor", UMeshPaintingToolProperties::StaticClass());
 	EraseColor->MarkHiddenByCustomization();
 
-	IDetailCategoryBuilder& WeightCategory = DetailLayout.EditCategory(TEXT("WeightPainting"));
-	WeightCategory.SetSortOrder(0);
-
 	TSharedRef<IPropertyHandle> WeightTypeProperty = DetailLayout.GetProperty(GET_MEMBER_NAME_CHECKED(UMeshVertexWeightPaintingToolProperties, TextureWeightType));
 	TSharedRef<IPropertyHandle> PaintWeightProperty = DetailLayout.GetProperty(GET_MEMBER_NAME_CHECKED(UMeshVertexWeightPaintingToolProperties, PaintTextureWeightIndex));
 	TSharedRef<IPropertyHandle> EraseWeightProperty = DetailLayout.GetProperty(GET_MEMBER_NAME_CHECKED(UMeshVertexWeightPaintingToolProperties, EraseTextureWeightIndex));
@@ -414,9 +419,11 @@ void FTexturePaintingSettingsCustomization::CustomizeDetails(IDetailLayoutBuilde
 	FMeshPaintingSettingsCustomization::CustomizeDetails(DetailLayout);
 
 	IDetailCategoryBuilder& ColorCategory = DetailLayout.EditCategory(TEXT("ColorPainting"));
-	ColorCategory.SetSortOrder(0);
-	IDetailCategoryBuilder& BrushCategory = DetailLayout.EditCategory(TEXT("Brush"), FText::GetEmpty(), ECategoryPriority::Important);
+	ColorCategory.SetSortOrder(1);
+	IDetailCategoryBuilder& BrushCategory = DetailLayout.EditCategory(TEXT("Brush"));
+	BrushCategory.SetSortOrder(2);
 	IDetailCategoryBuilder& TextureCategory = DetailLayout.EditCategory(TEXT("TexturePainting"));
+	TextureCategory.SetSortOrder(3);
 
 	TSharedRef<IPropertyHandle> RedChannel = DetailLayout.GetProperty("bWriteRed", UMeshTexturePaintingToolProperties::StaticClass());
 	RedChannel->MarkHiddenByCustomization();
@@ -476,6 +483,7 @@ void FTextureAssetPaintingSettingsCustomization::CustomizeDetails(IDetailLayoutB
 	FTexturePaintingSettingsCustomization::CustomizeDetails(DetailLayout);
 
 	IDetailCategoryBuilder& TextureCategory = DetailLayout.EditCategory(TEXT("TexturePainting"));
+	TextureCategory.SetSortOrder(3);
 
 	TSharedRef<IPropertyHandle> UVChannel = DetailLayout.GetProperty(GET_MEMBER_NAME_CHECKED(UMeshTextureAssetPaintingToolProperties, UVChannel));
 	UVChannel->MarkHiddenByCustomization();
diff --git a/Engine/Plugins/MeshPainting/Source/MeshPaintingToolset/Public/BaseMeshPaintingToolProperties.h b/Engine/Plugins/MeshPainting/Source/MeshPaintingToolset/Public/BaseMeshPaintingToolProperties.h
--- a/Engine/Plugins/MeshPainting/Source/MeshPaintingToolset/Public/BaseMeshPaintingToolProperties.h
+++ b/Engine/Plugins/MeshPainting/Source/MeshPaintingToolset/Public/BaseMeshPaintingToolProperties.h
@@ -12,11 +12,11 @@ class MESHPAINTINGTOOLSET_API UMeshPaintingToolProperties : public UBrushBasePro
 
 public:
 	/** Color used for applying color painting */
-	UPROPERTY(EditAnywhere, Category = Brush)
+	UPROPERTY(EditAnywhere, Category = ColorPainting)
 	FLinearColor PaintColor = FLinearColor::White;
 
 	/** Color used for erasing color painting */
-	UPROPERTY(EditAnywhere, Category = Brush)
+	UPROPERTY(EditAnywhere, Category = ColorPainting)
 	FLinearColor EraseColor = FLinearColor::Black;
 
 	/** Enables "Flow" painting where paint is continually applied from the brush every tick */
diff --git a/Engine/Plugins/MeshPainting/Source/MeshPaintingToolset/Public/MeshVertexPaintingTool.h b/Engine/Plugins/MeshPainting/Source/MeshPaintingToolset/Public/MeshVertexPaintingTool.h
--- a/Engine/Plugins/MeshPainting/Source/MeshPaintingToolset/Public/MeshVertexPaintingTool.h
+++ b/Engine/Plugins/MeshPainting/Source/MeshPaintingToolset/Public/MeshVertexPaintingTool.h
@@ -88,7 +88,7 @@ public:
 	int32 LODIndex = 0;
 
 	/** Size of vertex points drawn when mesh painting is active. */
-	UPROPERTY(EditAnywhere, Category = "VertexPainting")
+	UPROPERTY(EditAnywhere, Category = VertexPainting, meta = (UIMin = "0", ClampMin = "0"))
 	float VertexPreviewSize;
 };
 
