Fix race condition crash where MeshPaint UI requests texture information during texture compilation.

diff --git a/Engine/Plugins/MeshPainting/Source/MeshPaintingToolset/Private/MeshPaintHelpers.cpp b/Engine/Plugins/MeshPainting/Source/MeshPaintingToolset/Private/MeshPaintHelpers.cpp
--- a/Engine/Plugins/MeshPainting/Source/MeshPaintingToolset/Private/MeshPaintHelpers.cpp
+++ b/Engine/Plugins/MeshPainting/Source/MeshPaintingToolset/Private/MeshPaintHelpers.cpp
@@ -664,14 +664,14 @@ uint32 UMeshPaintingSubsystem::GetMeshPaintTextureResourceSize(UMeshComponent* M
 {
 	if (UTexture* Texture = MeshComponent->GetMeshPaintTexture())
 	{
-		FTexturePlatformData** PlatformDataPtr = Texture->GetRunningPlatformData();
-		if (PlatformDataPtr != nullptr && *PlatformDataPtr != nullptr)
+		// Check that the texture has finished compilation before reading platform data.
+		if (!Texture->IsDefaultTexture())
 		{
-			return (*PlatformDataPtr)->GetPayloadSize(0);
-		}
-		else
-		{
-			return Texture->GetResourceSizeBytes(EResourceSizeMode::Exclusive);
+			FTexturePlatformData** PlatformDataPtr = Texture->GetRunningPlatformData();
+			if (PlatformDataPtr != nullptr && *PlatformDataPtr != nullptr)
+			{
+				return (*PlatformDataPtr)->GetPayloadSize(0);
+			}
 		}
 	}
 	return 0;
