Fix Mesh Paint resetting all tool tabs after level save.
Level save will deactivate open tools. Previously if no mesh paint tool was active we defaulted to the vertex paint selection tool. Now we use the stored current tab to find the appropriate tool to default to.

diff --git a/Engine/Plugins/MeshPainting/Source/MeshPaintEditorMode/Private/MeshPaintMode.cpp b/Engine/Plugins/MeshPainting/Source/MeshPaintEditorMode/Private/MeshPaintMode.cpp
--- a/Engine/Plugins/MeshPainting/Source/MeshPaintEditorMode/Private/MeshPaintMode.cpp
+++ b/Engine/Plugins/MeshPainting/Source/MeshPaintEditorMode/Private/MeshPaintMode.cpp
@@ -228,37 +228,48 @@ void UMeshPaintMode::Tick(FEditorViewportClient* ViewportClient, float DeltaTime
 
 	FName ActiveTab = Toolkit->GetCurrentPalette();
 	FName TargetTab = ActiveTab;
-	EMeshPaintActiveMode CurrentActiveMode = EMeshPaintActiveMode::VertexColor;
 	
 	if (ActiveToolName == VertexColorPaintToolName)
 	{
 		TargetTab = MeshPaintMode_VertexColor;
-		CurrentActiveMode = EMeshPaintActiveMode::VertexColor;
 	}
 	else if (ActiveToolName == VertexWeightPaintToolName)
 	{
 		TargetTab = MeshPaintMode_VertexWeights;
-		CurrentActiveMode = EMeshPaintActiveMode::VertexWeights;
 	}
 	else if (ActiveToolName == TextureColorPaintToolName || ActiveToolName == TextureColorSelectToolName)
 	{
 		TargetTab = MeshPaintMode_TextureColor;
-		CurrentActiveMode = EMeshPaintActiveMode::TextureColor;
 	}
 	else if (ActiveToolName == TextureAssetPaintToolName || ActiveToolName == TextureAssetSelectToolName)
 	{
 		TargetTab = MeshPaintMode_TextureAsset;
-		CurrentActiveMode = EMeshPaintActiveMode::Texture;
 	}
-	else if (ActiveTool == nullptr)
+
+	EMeshPaintActiveMode CurrentActiveMode = EMeshPaintActiveMode::VertexColor;
+
+	if (TargetTab == MeshPaintMode_VertexColor)
 	{
-		TargetTab = MeshPaintMode_VertexColor;
 		CurrentActiveMode = EMeshPaintActiveMode::VertexColor;
 	}
-	if (TargetTab != ActiveTab)
+	else if (TargetTab == MeshPaintMode_VertexWeights)
+	{
+		CurrentActiveMode = EMeshPaintActiveMode::VertexWeights;
+	}
+	else if (TargetTab == MeshPaintMode_TextureColor)
+	{
+		CurrentActiveMode = EMeshPaintActiveMode::TextureColor;
+	}
+	else if (TargetTab == MeshPaintMode_TextureAsset)
+	{
+		CurrentActiveMode = EMeshPaintActiveMode::Texture;
+	}
+
+	if (TargetTab != ActiveTab || ActiveTool == nullptr)
 	{
 		Toolkit->SetCurrentPalette(TargetTab);
 	}
+
 	if (ViewportClient->IsPerspective())
 	{
 		// Make sure perspective viewports are still set to real-time
