Resize MeshPaintTexture in place instead of swapping for a new texture.
When we "Fix" MeshPaintTexture size we used to create a new texture object and set to the component. This means that the user can't use a single texture editor session to inspect the texture changes.
After the change we update the texture object in place.

diff --git a/Engine/Plugins/MeshPainting/Source/MeshPaintEditorMode/Private/MeshPaintModeHelpers.cpp b/Engine/Plugins/MeshPainting/Source/MeshPaintEditorMode/Private/MeshPaintModeHelpers.cpp
--- a/Engine/Plugins/MeshPainting/Source/MeshPaintEditorMode/Private/MeshPaintModeHelpers.cpp
+++ b/Engine/Plugins/MeshPainting/Source/MeshPaintEditorMode/Private/MeshPaintModeHelpers.cpp
@@ -958,13 +958,9 @@ void UMeshPaintModeSubsystem::FixTextureColors(const TArray<UMeshComponent*>& Co
 						FImage ResizedImage(TextureResolution, TextureResolution, Image.NumSlices, Image.Format, Image.GammaSpace);
 						FImageCore::ResizeImage(Image, ResizedImage);
 
-						UMeshPaintVirtualTexture* NewTexture = NewObject<UMeshPaintVirtualTexture>(StaticMeshComponent->GetOutermost());
-						NewTexture->Source.Init(ResizedImage);
-						NewTexture->OwningComponent = MakeWeakObjectPtr(StaticMeshComponent);
-						NewTexture->UpdateResource();
-
-						StaticMeshComponent->Modify();
-						StaticMeshComponent->SetMeshPaintTexture(NewTexture);
+						Texture->Modify();
+						Texture->Source.Init(ResizedImage);
+						Texture->PostEditChange();
 					}
 				}
 			}
