Fix mesh paint mode tools not sharing their brush settings.
Previously texture and vertex painting would maintain their own version of the brush settings.
The issue was that Save/Restore UInteractiveToolPropertySet properties cache results on the CDO. But the different mesh modes all use a different derived property class (and CDO).
The fix is to add new functions which can iterate the class hierachy and store properties on every CDO in it.
We _probably_ wouldn't want this in the general case, but it is what we want for the mesh paint mode.

diff --git a/Engine/Plugins/MeshPainting/Source/MeshPaintingToolset/Private/MeshPaintHelpers.cpp b/Engine/Plugins/MeshPainting/Source/MeshPaintingToolset/Private/MeshPaintHelpers.cpp
--- a/Engine/Plugins/MeshPainting/Source/MeshPaintingToolset/Private/MeshPaintHelpers.cpp
+++ b/Engine/Plugins/MeshPainting/Source/MeshPaintingToolset/Private/MeshPaintHelpers.cpp
@@ -9,6 +9,7 @@
 #include "Engine/SkinnedAssetCommon.h"
 #include "Engine/Texture2D.h"
 #include "IMeshPaintComponentAdapter.h"
+#include "InteractiveTool.h"
 #include "MeshPaintAdapterFactory.h"
 #include "MeshVertexPaintingTool.h"
 #include "Rendering/SkeletalMeshLODModel.h"
@@ -1714,3 +1715,56 @@ int32 UMeshPaintingSubsystem::GetMaxUVIndexToPaint() const
 	return 0;
 }
 
+TCHAR const* FMeshPaintToolSettingHelpers::GetCacheIdentifier()
+{
+	return TEXT("MeshPaint");
+}
+
+void FMeshPaintToolSettingHelpers::SavePropertiesForClassHeirachy(UInteractiveTool* InTool, UInteractiveToolPropertySet* InProperties)
+{
+	// Properties are saved in a cache on the class CDO. 
+	// So we need to iterate the class structure and save for each class in the hierachy.
+	// That way classes that share a base class will both use the same cache for their shared properties.
+	// Specifically we keep the same base class brush settings when switching between the different painting tools.
+	const FString CacheIdentifier(GetCacheIdentifier());
+	InProperties->SaveProperties(InTool, CacheIdentifier);
+
+	UEngine::FCopyPropertiesForUnrelatedObjectsParams Params;
+	Params.bOnlyHandleDirectSubObjects = true;
+
+	UClass* Class = InProperties->GetClass();
+	if (Class != UInteractiveToolPropertySet::StaticClass())
+	{
+		Class = Class->GetSuperClass();
+		while (Class != UInteractiveToolPropertySet::StaticClass())
+		{
+			UInteractiveToolPropertySet* Properties = NewObject<UInteractiveToolPropertySet>(GetTransientPackage(), Class);
+			UEngine::CopyPropertiesForUnrelatedObjects(InProperties, Properties, Params);
+			Properties->SaveProperties(InTool, CacheIdentifier);
+			Class = Class->GetSuperClass();
+		}
+	}
+}
+
+void FMeshPaintToolSettingHelpers::RestorePropertiesForClassHeirachy(UInteractiveTool* InTool, UInteractiveToolPropertySet* InProperties)
+{
+	// Restore from leaf class down the hierachy, with each base class overriding what was written before.
+	const FString CacheIdentifier(GetCacheIdentifier());
+	InProperties->RestoreProperties(InTool, CacheIdentifier);
+
+	UEngine::FCopyPropertiesForUnrelatedObjectsParams Params;
+	Params.bOnlyHandleDirectSubObjects = true;
+
+	UClass* Class = InProperties->GetClass();
+	if (Class != UInteractiveToolPropertySet::StaticClass())
+	{
+		Class = Class->GetSuperClass();
+		while (Class != UInteractiveToolPropertySet::StaticClass())
+		{
+			UInteractiveToolPropertySet* Properties = NewObject<UInteractiveToolPropertySet>(GetTransientPackage(), Class);
+			Properties->RestoreProperties(InTool, CacheIdentifier);
+			UEngine::CopyPropertiesForUnrelatedObjects(Properties, InProperties, Params);
+			Class = Class->GetSuperClass();
+		}
+	}
+}
diff --git a/Engine/Plugins/MeshPainting/Source/MeshPaintingToolset/Private/MeshTexturePaintingTool.cpp b/Engine/Plugins/MeshPainting/Source/MeshPaintingToolset/Private/MeshTexturePaintingTool.cpp
--- a/Engine/Plugins/MeshPainting/Source/MeshPaintingToolset/Private/MeshTexturePaintingTool.cpp
+++ b/Engine/Plugins/MeshPainting/Source/MeshPaintingToolset/Private/MeshTexturePaintingTool.cpp
@@ -72,10 +72,11 @@ void UMeshTexturePaintingTool::Setup()
 {
 	Super::Setup();
 
-	TextureProperties = Cast<UMeshTexturePaintingToolProperties>(BrushProperties);
 	bResultValid = false;
 	bStampPending = false;
-	BrushProperties->RestoreProperties(this);
+
+	FMeshPaintToolSettingHelpers::RestorePropertiesForClassHeirachy(this, BrushProperties);
+	TextureProperties = Cast<UMeshTexturePaintingToolProperties>(BrushProperties);
 
 	// Needed after restoring properties because the brush radius may be an output
 	// property based on selection, so we shouldn't use the last stored value there.
@@ -107,7 +108,7 @@ void UMeshTexturePaintingTool::Shutdown(EToolShutdownType ShutdownType)
 		MeshPaintingSubsystem->Refresh();
 	}
 
-	BrushProperties->SaveProperties(this);
+	FMeshPaintToolSettingHelpers::SavePropertiesForClassHeirachy(this, BrushProperties);
 
 	Super::Shutdown(ShutdownType);
 }
diff --git a/Engine/Plugins/MeshPainting/Source/MeshPaintingToolset/Private/MeshVertexPaintingTool.cpp b/Engine/Plugins/MeshPainting/Source/MeshPaintingToolset/Private/MeshVertexPaintingTool.cpp
--- a/Engine/Plugins/MeshPainting/Source/MeshPaintingToolset/Private/MeshVertexPaintingTool.cpp
+++ b/Engine/Plugins/MeshPainting/Source/MeshPaintingToolset/Private/MeshVertexPaintingTool.cpp
@@ -64,10 +64,12 @@ bool UMeshVertexPaintingTool::IsMeshAdapterSupported(TSharedPtr<IMeshPaintCompon
 void UMeshVertexPaintingTool::Setup()
 {
 	Super::Setup();
-	VertexProperties = Cast<UMeshVertexPaintingToolProperties>(BrushProperties);
+
 	bResultValid = false;
 	bStampPending = false;
-	BrushProperties->RestoreProperties(this);
+
+	FMeshPaintToolSettingHelpers::RestorePropertiesForClassHeirachy(this, BrushProperties);
+	VertexProperties = Cast<UMeshVertexPaintingToolProperties>(BrushProperties);
 
 	// Needed after restoring properties because the brush radius may be an output
 	// property based on selection, so we shouldn't use the last stored value there.
@@ -89,12 +91,15 @@ void UMeshVertexPaintingTool::Shutdown(EToolShutdownType ShutdownType)
 	ApplyForcedLODIndex(-1);
 
 	FinishPainting();
-	BrushProperties->SaveProperties(this);
+	
 	UMeshPaintingSubsystem* MeshPaintingSubsystem = GEngine->GetEngineSubsystem<UMeshPaintingSubsystem>();
 	if (MeshPaintingSubsystem)
 	{
 		MeshPaintingSubsystem->Refresh();
 	}
+
+	FMeshPaintToolSettingHelpers::SavePropertiesForClassHeirachy(this, BrushProperties);
+
 	Super::Shutdown(ShutdownType);
 }
 
diff --git a/Engine/Plugins/MeshPainting/Source/MeshPaintingToolset/Public/MeshPaintHelpers.h b/Engine/Plugins/MeshPainting/Source/MeshPaintingToolset/Public/MeshPaintHelpers.h
--- a/Engine/Plugins/MeshPainting/Source/MeshPaintingToolset/Public/MeshPaintHelpers.h
+++ b/Engine/Plugins/MeshPainting/Source/MeshPaintingToolset/Public/MeshPaintHelpers.h
@@ -30,6 +30,8 @@ struct FStaticMeshComponentLODInfo;
 class UMeshVertexPaintingToolProperties;
 class UBrushBaseProperties;
 struct FMeshDescription;
+class UInteractiveTool;
+class UInteractiveToolPropertySet;
 
 enum class EMeshPaintDataColorViewMode : uint8;
 
@@ -81,6 +83,18 @@ struct FInstanceTexturePaintSettings
 	}
 };
 
+/** Struct for some static helper functions to store tool properties. */
+class FMeshPaintToolSettingHelpers
+{
+public:
+	/** Store properties in a property set for later restore in a session. This saves settings per base class so that different class hierachies can share their base class data. */
+	static void SavePropertiesForClassHeirachy(UInteractiveTool* InTool, UInteractiveToolPropertySet* InProperties);
+	/** Restore properties in a property that were saved in this session. This restores settings per base class so that different class hierachies can share their base class data. */
+	static void RestorePropertiesForClassHeirachy(UInteractiveTool* InTool, UInteractiveToolPropertySet* InProperties);
+private:
+	/** Get unique string used to identify property set in cache. */
+	static TCHAR const* GetCacheIdentifier();
+};
 
 UENUM()
 enum class ETexturePaintWeightTypes : uint8
@@ -125,8 +139,8 @@ DECLARE_DELEGATE_ThreeParams(FPerTrianglePaintAction, IMeshPaintComponentAdapter
 UCLASS()
 class MESHPAINTINGTOOLSET_API UMeshPaintingSubsystem : public UEngineSubsystem
 {
-
 	GENERATED_BODY()
+
 public:
 	UMeshPaintingSubsystem();
 
