Small QOL improvement to remember last used mesh paint mode tab.
Store the last used tab in editor settings whenever exiting the mesh paint mode.Apply it whenever entering mode.

diff --git a/Engine/Plugins/MeshPainting/Source/MeshPaintEditorMode/Private/MeshPaintMode.cpp b/Engine/Plugins/MeshPainting/Source/MeshPaintEditorMode/Private/MeshPaintMode.cpp
--- a/Engine/Plugins/MeshPainting/Source/MeshPaintEditorMode/Private/MeshPaintMode.cpp
+++ b/Engine/Plugins/MeshPainting/Source/MeshPaintEditorMode/Private/MeshPaintMode.cpp
@@ -52,6 +52,18 @@ UMeshPaintMode* UMeshPaintMode::GetMeshPaintMode()
 	return Cast<UMeshPaintMode>(GLevelEditorModeTools().GetActiveScriptableMode("MeshPaintMode"));
 }
 
+FName UMeshPaintMode::GetValidPaletteName(const FName InName)
+{
+	if (InName == UMeshPaintMode::MeshPaintMode_VertexColor ||
+		InName == UMeshPaintMode::MeshPaintMode_VertexWeights ||
+		InName == UMeshPaintMode::MeshPaintMode_TextureColor ||
+		InName == UMeshPaintMode::MeshPaintMode_TextureAsset)
+	{
+		return InName;
+	}
+	return UMeshPaintMode::MeshPaintMode_VertexColor;
+}
+
 template<typename T>
 static T* GetTypedToolProperties()
 {
@@ -160,7 +172,7 @@ void UMeshPaintMode::Enter()
 
 	// disable tool change tracking to activate default tool
 	GetToolManager()->ConfigureChangeTrackingMode(EToolChangeTrackingMode::NoChangeTracking);
-	Toolkit->SetCurrentPalette(MeshPaintMode_VertexColor);
+	Toolkit->SetCurrentPalette(GetValidPaletteName(ModeSettings->DefaultPalette));
 	// switch back to full undo / redo tracking mode here if that is behavior we want
 	//GetToolManager()->ConfigureChangeTrackingMode(EToolChangeTrackingMode::FullUndoRedo);
 
@@ -173,6 +185,8 @@ void UMeshPaintMode::Enter()
 
 void UMeshPaintMode::Exit()
 {
+	ModeSettings->DefaultPalette = Toolkit->GetCurrentPalette();
+
 	Toolkit->OnPaletteChanged().Remove(PaletteChangedHandle);
 	FCoreUObjectDelegates::OnObjectsReplaced.RemoveAll(this);
 	GEditor->OnEditorClose().RemoveAll(this);
diff --git a/Engine/Plugins/MeshPainting/Source/MeshPaintEditorMode/Private/MeshPaintMode.h b/Engine/Plugins/MeshPainting/Source/MeshPaintEditorMode/Private/MeshPaintMode.h
--- a/Engine/Plugins/MeshPainting/Source/MeshPaintEditorMode/Private/MeshPaintMode.h
+++ b/Engine/Plugins/MeshPainting/Source/MeshPaintEditorMode/Private/MeshPaintMode.h
@@ -44,6 +44,7 @@ public:
 	static UMeshTextureColorPaintingToolProperties* GetTextureColorToolProperties();
 	static UMeshTextureAssetPaintingToolProperties* GetTextureAssetToolProperties();
 	static UMeshPaintMode* GetMeshPaintMode();
+	static FName GetValidPaletteName(const FName InName);
 
 	static FName MeshPaintMode_VertexColor;
 	static FName MeshPaintMode_VertexWeights;
diff --git a/Engine/Plugins/MeshPainting/Source/MeshPaintEditorMode/Private/MeshPaintModeSettings.h b/Engine/Plugins/MeshPainting/Source/MeshPaintEditorMode/Private/MeshPaintModeSettings.h
--- a/Engine/Plugins/MeshPainting/Source/MeshPaintEditorMode/Private/MeshPaintModeSettings.h
+++ b/Engine/Plugins/MeshPainting/Source/MeshPaintEditorMode/Private/MeshPaintModeSettings.h
@@ -42,4 +42,8 @@ public:
 	/** Color view mode used to display Vertex Colors */
 	UPROPERTY(config, EditAnywhere, Category = Visualization)
 	EMeshPaintDataColorViewMode ColorViewMode = EMeshPaintDataColorViewMode::Normal;
+
+	/** Palette tab to use when entering the tool */
+	UPROPERTY(config)
+	FName DefaultPalette;
 };
