Mesh texture paint support textures imported as single channel 8 bit.
Previously the texture data copying functions all were direct memory copies with assumptions about the layout.
Changed these to use FImage which handles the underlying details.

diff --git a/Engine/Plugins/MeshPainting/Source/MeshPaintingToolset/Private/MeshTexturePaintingTool.cpp b/Engine/Plugins/MeshPainting/Source/MeshPaintingToolset/Private/MeshTexturePaintingTool.cpp
--- a/Engine/Plugins/MeshPainting/Source/MeshPaintingToolset/Private/MeshTexturePaintingTool.cpp
+++ b/Engine/Plugins/MeshPainting/Source/MeshPaintingToolset/Private/MeshTexturePaintingTool.cpp
@@ -1224,11 +1224,8 @@ void UMeshTexturePaintingTool::FinishPaintingTexture()
 			TextureData->PaintingTexture2D->PreEditChange(nullptr);
 
 			// Store source art
-			// @@ use FMipLock.Image instead
-			FColor* Colors = (FColor*)TextureData->PaintingTexture2D->Source.LockMip(0);
-			check(TextureData->PaintingTexture2D->Source.CalcMipSize(0) == TexturePixels.Num() * sizeof(FColor));
-			FMemory::Memcpy(Colors, TexturePixels.GetData(), TexturePixels.Num() * sizeof(FColor));
-			TextureData->PaintingTexture2D->Source.UnlockMip(0);
+			FImageView ImageView(TexturePixels.GetData(), TexWidth, TexHeight, EGammaSpace::sRGB);
+			TextureData->PaintingTexture2D->Source.Init(ImageView);
 
 			TextureData->PaintingTexture2D->bHasBeenPaintedInEditor = true;
 
diff --git a/Engine/Plugins/MeshPainting/Source/MeshPaintingToolset/Private/TexturePaintToolset.cpp b/Engine/Plugins/MeshPainting/Source/MeshPaintingToolset/Private/TexturePaintToolset.cpp
--- a/Engine/Plugins/MeshPainting/Source/MeshPaintingToolset/Private/TexturePaintToolset.cpp
+++ b/Engine/Plugins/MeshPainting/Source/MeshPaintingToolset/Private/TexturePaintToolset.cpp
@@ -305,44 +305,19 @@ UTexture2D* UTexturePaintToolset::CreateScratchUncompressedTexture(UTexture2D* S
 {
 	check(SourceTexture->Source.IsValid());
 
-	// bad: this code assumes the TextureSource is RGBA8
-	//	use GetMipImage instead of GetMipData
-
-	// Decompress PNG image
-	TArray64<uint8> RawData;
-	SourceTexture->Source.GetMipData(RawData, 0);
-
-	// We are using the source art so grab the original width/height
-	const int32 Width = SourceTexture->Source.GetSizeX();
-	const int32 Height = SourceTexture->Source.GetSizeY();
-	const bool bUseSRGB = SourceTexture->SRGB;
-
-	check(Width > 0 && Height > 0 && RawData.Num() > 0);
+	// Decompress PNG image and convert to BGRA8 for painting.
+	FImage SourceImage;
+	SourceTexture->Source.GetMipImage(SourceImage, 0);
+	SourceImage.ChangeFormat(ERawImageFormat::BGRA8, SourceImage.GetGammaSpace());
 
 	// Allocate the new texture
-	UTexture2D* NewTexture2D = UTexture2D::CreateTransient(Width, Height, GetTempUncompressedTexturePixelFormat());
+	UTexture2D* NewTexture2D = NewObject<UTexture2D>(GetTransientPackage(), NAME_None, RF_Transient | RF_Transactional);
 
-	NewTexture2D->SetFlags(RF_Transactional);
 	// Fill in the base mip for the texture we created
-	uint8* MipData = (uint8*)NewTexture2D->GetPlatformData()->Mips[0].BulkData.Lock(LOCK_READ_WRITE);
-	for (int32 y = 0; y < Height; y++)
-	{
-		uint8* DestPtr = &MipData[(Height - 1 - y) * Width * sizeof(FColor)];
-		//bad: very unsafe cast and un-checked access of memory :
-		const FColor* SrcPtr = &((FColor*)(RawData.GetData()))[(Height - 1 - y) * Width];
-		for (int32 x = 0; x < Width; x++)
-		{
-			*DestPtr++ = SrcPtr->B;
-			*DestPtr++ = SrcPtr->G;
-			*DestPtr++ = SrcPtr->R;
-			*DestPtr++ = SrcPtr->A;
-			SrcPtr++;
-		}
-	}
-	NewTexture2D->GetPlatformData()->Mips[0].BulkData.Unlock();
+	NewTexture2D->Source.Init(SourceImage);
 
 	// Set options
-	NewTexture2D->SRGB = bUseSRGB;
+	NewTexture2D->SRGB = SourceTexture->SRGB;
 	NewTexture2D->CompressionNone = true;
 	NewTexture2D->MipGenSettings = TMGS_NoMipmaps;
 	NewTexture2D->CompressionSettings = TC_Default;
@@ -368,6 +343,11 @@ void UTexturePaintToolset::SetupInitialRenderTargetData(UTexture2D* InTextureSou
 			UTexture2D* TempSourceArtTexture = CreateScratchUncompressedTexture(InTextureSource);
 			check(TempSourceArtTexture != nullptr);
 
+#if WITH_EDITOR
+			// We need to complete texture compilation before we can copy to render target.
+			TempSourceArtTexture->BlockOnAnyAsyncBuild();
+#endif
+
 			// Copy the texture to the render target using the GPU
 			CopyTextureToRenderTargetTexture(TempSourceArtTexture, InRenderTarget, GEditor->GetEditorWorldContext().World()->GetFeatureLevel());
 
@@ -499,13 +479,11 @@ void UTexturePaintToolset::RetrieveTexturesForComponent(const UMeshComponent* Co
 		{
 			if (const UTexture* Texture = PaintableTexture.Texture)
 			{
-				// Only BRGA8 image are supported by the tool currently (we should probably provide some feedback to the user on why some texture are not supported and a way to convert them)
-				return Texture->Source.GetFormat() != TSF_BGRA8;
+				// Only 8 bit images are supported by the tool currently (we should probably provide some feedback to the user on why some texture are not supported and a way to convert them)
+				return Texture->Source.GetFormat() != TSF_BGRA8 && Texture->Source.GetFormat() != TSF_G8;
 			}
 
 			return true;
 		});
 	}
 }
-
-
