Improve loading time for mesh paint texture select widget.
Comparing objects forces a load so compare paths instead.

diff --git a/Engine/Plugins/MeshPainting/Source/MeshPaintingToolset/Private/MeshTexturePaintingTool.cpp b/Engine/Plugins/MeshPainting/Source/MeshPaintingToolset/Private/MeshTexturePaintingTool.cpp
--- a/Engine/Plugins/MeshPainting/Source/MeshPaintingToolset/Private/MeshTexturePaintingTool.cpp
+++ b/Engine/Plugins/MeshPainting/Source/MeshPaintingToolset/Private/MeshTexturePaintingTool.cpp
@@ -1462,12 +1462,8 @@ void UMeshTextureAssetPaintingTool::GetModifiedTexturesToSave(TArray<UObject*>&
 
 bool UMeshTextureAssetPaintingTool::ShouldFilterTextureAsset(const FAssetData& AssetData) const
 {
-	UTexture2D* AssetTexture = Cast<UTexture2D>(AssetData.GetAsset());
-	if (!AssetTexture)
-	{
-		return true;
-	}
-	return !(PaintableTextures.ContainsByPredicate([=](const FPaintableTexture& Texture) { return Texture.Texture->GetFullName() == AssetTexture->GetFullName(); }));
+	FSoftObjectPath Path = AssetData.GetSoftObjectPath();
+	return !(PaintableTextures.ContainsByPredicate([&Path](const FPaintableTexture& Texture) { return FSoftObjectPath(Texture.Texture) == Path; }));
 }
 
 void UMeshTextureAssetPaintingTool::OnPropertyModified(UObject* PropertySet, FProperty* Property)
