MeshPaint: Remove DoesMaterialUseTexture() check before applying material override.
The logic doesn't cover all cases (layers, parameters in material functions etc) so sometimes the override would not  be applied.
We call RegisterMaterialOverride() which does apply the correct logic anyway, so the filtering is redundant.

diff --git a/Engine/Plugins/MeshPainting/Source/MeshPaintingToolset/Private/IMeshPaintComponentAdapter.cpp b/Engine/Plugins/MeshPainting/Source/MeshPaintingToolset/Private/IMeshPaintComponentAdapter.cpp
--- a/Engine/Plugins/MeshPainting/Source/MeshPaintingToolset/Private/IMeshPaintComponentAdapter.cpp
+++ b/Engine/Plugins/MeshPainting/Source/MeshPaintingToolset/Private/IMeshPaintComponentAdapter.cpp
@@ -424,7 +424,7 @@ namespace UE::MeshPaintingToolset
 				// Unregister for all materials. This will no affect the material that weren't overridden by this instance
 				Private::FGlobalTextureOverrideState::RemoveMaterialOverride(this, MaterialToCheck, SourceTexture);
 			}
-			else if (DoesMaterialUseTexture(MaterialToCheck, SourceTexture))
+			else
 			{
 				// Keep track of the material overridden
 				Private::FGlobalTextureOverrideState::RegisterMaterialOverride(this, MaterialToCheck, SourceTexture, OverrideTexture, FeatureLevel);
