Add ElementSelectionChangeNotify() to UEdMode.
Use it to route selection changes to MeshPaint tool.
This allows MeshPaint tool to respond to per component selection as well as per actor selection.
Fixes the case where a user wants to select one of multiple mesh components from a single actor for painting.

diff --git a/Engine/Plugins/MeshPainting/Source/MeshPaintEditorMode/Private/MeshPaintMode.cpp b/Engine/Plugins/MeshPainting/Source/MeshPaintEditorMode/Private/MeshPaintMode.cpp
--- a/Engine/Plugins/MeshPainting/Source/MeshPaintEditorMode/Private/MeshPaintMode.cpp
+++ b/Engine/Plugins/MeshPainting/Source/MeshPaintEditorMode/Private/MeshPaintMode.cpp
@@ -451,6 +451,11 @@ void UMeshPaintMode::ActorSelectionChangeNotify()
 	UpdateSelectedMeshes();
 }
 
+void UMeshPaintMode::ElementSelectionChangeNotify()
+{
+	UpdateSelectedMeshes();
+}
+
 void UMeshPaintMode::ActorPropChangeNotify()
 {
 	// Setting change on selected components can change whether they are valid for painting.
diff --git a/Engine/Plugins/MeshPainting/Source/MeshPaintEditorMode/Private/MeshPaintMode.h b/Engine/Plugins/MeshPainting/Source/MeshPaintEditorMode/Private/MeshPaintMode.h
--- a/Engine/Plugins/MeshPainting/Source/MeshPaintEditorMode/Private/MeshPaintMode.h
+++ b/Engine/Plugins/MeshPainting/Source/MeshPaintEditorMode/Private/MeshPaintMode.h
@@ -75,6 +75,7 @@ protected:
 	virtual void OnToolStarted(UInteractiveToolManager* Manager, UInteractiveTool* Tool) override;
 	virtual void OnToolEnded(UInteractiveToolManager* Manager, UInteractiveTool* Tool) override;
 	virtual void ActorSelectionChangeNotify() override;
+	virtual void ElementSelectionChangeNotify() override;
 	virtual void ActorPropChangeNotify() override;
 	virtual void ActivateDefaultTool() override;
 	virtual void UpdateOnPaletteChange(FName NewPalette);
diff --git a/Engine/Source/Editor/LevelEditor/Private/SLevelEditor.cpp b/Engine/Source/Editor/LevelEditor/Private/SLevelEditor.cpp
--- a/Engine/Source/Editor/LevelEditor/Private/SLevelEditor.cpp
+++ b/Engine/Source/Editor/LevelEditor/Private/SLevelEditor.cpp
@@ -2016,6 +2016,8 @@ void SLevelEditor::OnElementSelectionChanged(const UTypedElementSelectionSet* Se
 		}
 	}
 
+	GetEditorModeManager().ElementSelectionChangeNotify();
+
 	CachedViewportContextMenuTitle = FLevelEditorContextMenu::GetContextMenuTitle(ELevelEditorMenuContext::MainMenu, SelectionSet);
 
 #if PLATFORM_MAC
diff --git a/Engine/Source/Editor/UnrealEd/Private/EditorModeManager.cpp b/Engine/Source/Editor/UnrealEd/Private/EditorModeManager.cpp
--- a/Engine/Source/Editor/UnrealEd/Private/EditorModeManager.cpp
+++ b/Engine/Source/Editor/UnrealEd/Private/EditorModeManager.cpp
@@ -532,6 +532,15 @@ void FEditorModeTools::ActorSelectionChangeNotify()
 	});
 }
 
+void FEditorModeTools::ElementSelectionChangeNotify()
+{
+	ForEachEdMode([](UEdMode* Mode)
+	{
+		Mode->ElementSelectionChangeNotify();
+		return true;
+	});
+}
+
 void FEditorModeTools::ActorPropChangeNotify()
 {
 	ForEachEdMode([](UEdMode* Mode)
diff --git a/Engine/Source/Editor/UnrealEd/Public/EditorModeManager.h b/Engine/Source/Editor/UnrealEd/Public/EditorModeManager.h
--- a/Engine/Source/Editor/UnrealEd/Public/EditorModeManager.h
+++ b/Engine/Source/Editor/UnrealEd/Public/EditorModeManager.h
@@ -527,6 +527,7 @@ public:
 	UNREALED_API void ActorsDuplicatedNotify(TArray<AActor*>& InPreDuplicateSelection, TArray<AActor*>& InPostDuplicateSelection, const bool bOffsetLocations);
 	UNREALED_API void ActorMoveNotify();
 	UNREALED_API void ActorSelectionChangeNotify();
+	UNREALED_API void ElementSelectionChangeNotify();
 	UNREALED_API void ActorPropChangeNotify();
 	UNREALED_API void UpdateInternalData();
 	UNREALED_API bool IsOnlyVisibleActiveMode(FEditorModeID InMode) const;
diff --git a/Engine/Source/Editor/UnrealEd/Public/Tools/UEdMode.h b/Engine/Source/Editor/UnrealEd/Public/Tools/UEdMode.h
--- a/Engine/Source/Editor/UnrealEd/Public/Tools/UEdMode.h
+++ b/Engine/Source/Editor/UnrealEd/Public/Tools/UEdMode.h
@@ -96,6 +96,7 @@ public:
 	virtual void ActorMoveNotify() {}
 	virtual void ActorsDuplicatedNotify(TArray<AActor*>& PreDuplicateSelection, TArray<AActor*>& PostDuplicateSelection, bool bOffsetLocations) {}
 	virtual void ActorSelectionChangeNotify() {};
+	virtual void ElementSelectionChangeNotify() {};
 	virtual void ActorPropChangeNotify() {}
 	virtual void MapChangeNotify() {}
 
