MeshPaint: Iterate material functions to find textures to paint on, not just the top level material.
Previously we could still find some textures from the material functions via the iteration of texture parameters. But that wasn't exhaustive.
Also we only determine if a texture is marked as the default texture to paint when we iterate the expressions.
Now that we do that for material functions we can find default textures from there.

diff --git a/Engine/Plugins/MeshPainting/Source/MeshPaintingToolset/Private/IMeshPaintComponentAdapter.cpp b/Engine/Plugins/MeshPainting/Source/MeshPaintingToolset/Private/IMeshPaintComponentAdapter.cpp
--- a/Engine/Plugins/MeshPainting/Source/MeshPaintingToolset/Private/IMeshPaintComponentAdapter.cpp
+++ b/Engine/Plugins/MeshPainting/Source/MeshPaintingToolset/Private/IMeshPaintComponentAdapter.cpp
@@ -51,67 +51,84 @@ static bool IsTextureSuitableForTexturePainting(const TWeakObjectPtr<UTexture> T
 		(TexturePtr->Source.GetBytesPerPixel() <= UTexturePaintToolset::GetMaxSupportedBytesPerPixelForPainting())); // Textures' sources must fit in FColor struct to be supported.
 }
 
+static void AddPaintableTextureFromMaterialExpression(UMaterialInterface const* InMaterial, UMaterialExpressionTextureBase const* InExpression, int32& OutDefaultIndex, TArray<struct FPaintableTexture>& InOutTextureList)
+{
+	if (InExpression == nullptr || !IsTextureSuitableForTexturePainting(InExpression->Texture))
+	{
+		return;
+	}
+
+	// Default UV channel to index 0. 
+	FPaintableTexture PaintableTexture(InExpression->Texture, 0);
+
+	// Texture Samples can have UV's specified, check the first node for whether it has a custom UV channel set. 
+	// We only check the first as the Mesh paint mode does not support painting with UV's modified in the shader.
+	if (UMaterialExpressionTextureSample const* TextureSampleExpression = Cast<UMaterialExpressionTextureSample>(InExpression))
+	{
+		if (UMaterialExpressionTextureCoordinate const* TextureCoordsExpression = Cast<UMaterialExpressionTextureCoordinate>(TextureSampleExpression->Coordinates.Expression))
+		{
+			// Store the uv channel, this is set when the texture is selected. 
+			PaintableTexture.UVChannelIndex = TextureCoordsExpression->CoordinateIndex;
+		}
+		else
+		{
+			PaintableTexture.UVChannelIndex = TextureSampleExpression->ConstCoordinate;
+		}
+
+		// Handle texture parameter expressions.
+		if (UMaterialExpressionTextureSampleParameter const* TextureSampleParameterExpression = Cast<UMaterialExpressionTextureSampleParameter>(TextureSampleExpression))
+		{
+			// Grab the overridden texture if it exists.  
+			InMaterial->GetTextureParameterValue(TextureSampleParameterExpression->ParameterName, PaintableTexture.Texture);
+		}
+	}
+
+	// Note that the same texture will be added again if its UV channel differs. 
+	const int32 TextureIndex = InOutTextureList.AddUnique(PaintableTexture);
+
+	// Cache the first default index, if there is no previous info this will be used as the selected texture.
+	if ((OutDefaultIndex == INDEX_NONE) && InExpression->IsDefaultMeshpaintTexture)
+	{
+		OutDefaultIndex = TextureIndex;
+	}
+}
+
 void IMeshPaintComponentAdapter::DefaultQueryPaintableTextures(int32 MaterialIndex, const UMeshComponent* MeshComponent, int32& OutDefaultIndex, TArray<struct FPaintableTexture>& InOutTextureList)
 {
 	OutDefaultIndex = INDEX_NONE;
 
 	// We already know the material we are painting on, take it off the static mesh component
-	UMaterialInterface* Material = MeshComponent->GetMaterial(MaterialIndex);
+	UMaterialInterface const* Material = MeshComponent->GetMaterial(MaterialIndex);
 
 	while (Material != nullptr)
 	{
 		if (Material != NULL)
 		{
-			FPaintableTexture PaintableTexture;
-			// Find all the unique textures used in the top material level of the selected actor materials
-
-			// Only grab the textures from the top level of samples
-			for (UMaterialExpression* Expression : Material->GetMaterial()->GetExpressions())
+			// First iterate top level material expressions.
+			for (UMaterialExpression* MaterialExpression : Material->GetMaterial()->GetExpressions())
 			{
-				UMaterialExpressionTextureBase* TextureBase = Cast<UMaterialExpressionTextureBase>(Expression);
-				if (TextureBase != NULL &&
-					IsTextureSuitableForTexturePainting(TextureBase->Texture))
+				if (UMaterialExpressionTextureBase* TextureExpression = Cast<UMaterialExpressionTextureBase>(MaterialExpression))
 				{
-					// Default UV channel to index 0. 
-					PaintableTexture = FPaintableTexture(TextureBase->Texture, 0);
-
-					// Texture Samples can have UV's specified, check the first node for whether it has a custom UV channel set. 
-					// We only check the first as the Mesh paint mode does not support painting with UV's modified in the shader.
-					UMaterialExpressionTextureSample* TextureSample = Cast<UMaterialExpressionTextureSample>(Expression);
-					if (TextureSample != NULL)
-					{
-						UMaterialExpressionTextureCoordinate* TextureCoords = Cast<UMaterialExpressionTextureCoordinate>(TextureSample->Coordinates.Expression);
-						if (TextureCoords != NULL)
-						{
-							// Store the uv channel, this is set when the texture is selected. 
-							PaintableTexture.UVChannelIndex = TextureCoords->CoordinateIndex;
-						}
-						else
-						{
-							PaintableTexture.UVChannelIndex = TextureSample->ConstCoordinate;
-						}
-
-						// Handle texture parameter expressions
-						UMaterialExpressionTextureSampleParameter* TextureSampleParameter = Cast<UMaterialExpressionTextureSampleParameter>(TextureSample);
-						if (TextureSampleParameter != NULL)
-						{
-							// Grab the overridden texture if it exists.  
-							Material->GetTextureParameterValue(TextureSampleParameter->ParameterName, PaintableTexture.Texture);
-						}
-					}
-
-					// note that the same texture will be added again if its UV channel differs. 
-					int32 TextureIndex = InOutTextureList.AddUnique(PaintableTexture);
+					AddPaintableTextureFromMaterialExpression(Material, TextureExpression, OutDefaultIndex, InOutTextureList);
+				}
+			}
 
-					// cache the first default index, if there is no previous info this will be used as the selected texture
-					if ((OutDefaultIndex == INDEX_NONE) && TextureBase->IsDefaultMeshpaintTexture)
+			// Now iterate material expressions from material functions.
+			TArray<UMaterialFunctionInterface*> MaterialFunctions;
+			Material->GetDependentFunctions(MaterialFunctions);
+			for (UMaterialFunctionInterface* MaterialFunction : MaterialFunctions)
+			{
+				for (UMaterialExpression* MaterialExpression : MaterialFunction->GetExpressions())
+				{
+					if (UMaterialExpressionTextureBase* TextureExpression = Cast<UMaterialExpressionTextureBase>(MaterialExpression))
 					{
-						OutDefaultIndex = TextureIndex;
+						AddPaintableTextureFromMaterialExpression(Material, TextureExpression, OutDefaultIndex, InOutTextureList);
 					}
 				}
 			}
 		}
-		// Make sure to include all texture parameters, this will include all of the texture parameters from internal material functions
+		
+		// Make sure to include all texture parameters.
 		TMap<FMaterialParameterInfo, FMaterialParameterMetadata> ParameterValues;
 		Material->GetAllParametersOfType(EMaterialParameterType::Texture, ParameterValues);
 
@@ -129,8 +146,7 @@ void IMeshPaintComponentAdapter::DefaultQueryPaintableTextures(int32 MaterialInd
 			}
 		}
 
-		UMaterialInstance* MaterialInstance = Cast<UMaterialInstance>(Material);
-		if (MaterialInstance)
+		if (UMaterialInstance const* MaterialInstance = Cast<UMaterialInstance>(Material))
 		{
 			Material = MaterialInstance->Parent ? Cast<UMaterialInstance>(MaterialInstance->Parent) : nullptr;
 		}
