MeshPaint: Only create render targets for textures that are currently paintable.
Before this change, when texture asset painting, we would create a render target for all paintable textures (ie all textures found textures in materials) rather than just the selected one.

diff --git a/Engine/Plugins/MeshPainting/Source/MeshPaintingToolset/Private/MeshTexturePaintingTool.cpp b/Engine/Plugins/MeshPainting/Source/MeshPaintingToolset/Private/MeshTexturePaintingTool.cpp
--- a/Engine/Plugins/MeshPainting/Source/MeshPaintingToolset/Private/MeshTexturePaintingTool.cpp
+++ b/Engine/Plugins/MeshPainting/Source/MeshPaintingToolset/Private/MeshTexturePaintingTool.cpp
@@ -1267,73 +1267,89 @@ void UMeshTexturePaintingTool::ClearAllTextureOverrides()
 
 void UMeshTexturePaintingTool::SetAllTextureOverrides()
 {
+	UMeshPaintingSubsystem* MeshPaintingSubsystem = GEngine->GetEngineSubsystem<UMeshPaintingSubsystem>();
+	TArray<UMeshComponent*> SelectedMeshComponents = MeshPaintingSubsystem->GetSelectedMeshComponents();
+
 	for (FPaintableTexture const& PaintableTexture : PaintableTextures)
 	{
-		if (UTexture2D* Texture2D = Cast<UTexture2D>(PaintableTexture.Texture))
+		// Apply the overrides only to the components that we are painting with this texture.
+		TArray<UMeshComponent*, TInlineAllocator<8>> PaintableMeshComponents;
+		for (UMeshComponent* MeshComponent : SelectedMeshComponents)
 		{
-			Texture2D->BlockOnAnyAsyncBuild();
-
-			// Create the render target to paint on.
-			FPaintTexture2DData* TextureData = AddPaintTargetData(Texture2D);
-
-			const int32 TextureWidth = Texture2D->Source.GetSizeX();
-			const int32 TextureHeight = Texture2D->Source.GetSizeY();
-
-			if (TextureData->PaintRenderTargetTexture == nullptr ||
-				TextureData->PaintRenderTargetTexture->GetSurfaceWidth() != TextureWidth ||
-				TextureData->PaintRenderTargetTexture->GetSurfaceHeight() != TextureHeight)
+			if (CanPaintTextureToComponent(PaintableTexture.Texture, MeshComponent))
 			{
-				TextureData->PaintRenderTargetTexture = NewObject<UTextureRenderTarget2D>(GetTransientPackage(), NAME_None, RF_Transient);
-				TextureData->PaintRenderTargetTexture->bNeedsTwoCopies = true;
-				const bool bForceLinearGamma = true;
-				TextureData->PaintRenderTargetTexture->InitCustomFormat(TextureWidth, TextureHeight, PF_A16B16G16R16, bForceLinearGamma);
-				TextureData->PaintRenderTargetTexture->UpdateResourceImmediate();
-
-				TextureData->PaintRenderTargetTextureAdapter = nullptr; 
-				if (TextureData->PaintingTexture2D->IsCurrentlyVirtualTextured())
+				TSharedPtr<IMeshPaintComponentAdapter> MeshAdapter = MeshPaintingSubsystem->GetAdapterForComponent(MeshComponent);
+				if (IsMeshAdapterSupported(MeshAdapter))
 				{
-					// Virtual textures can't just swap in a render target in their material, so we use a virtual texture adapter.
-					FVirtualTextureBuildSettings VirtualTextureBuildSettings;
-					TextureData->PaintingTexture2D->GetVirtualTextureBuildSettings(VirtualTextureBuildSettings);
-
-					TextureData->PaintRenderTargetTextureAdapter = NewObject<UVirtualTextureAdapter>(GetTransientPackage(), NAME_None, RF_Transient);
-					TextureData->PaintRenderTargetTextureAdapter->Texture = TextureData->PaintRenderTargetTexture;
-					TextureData->PaintRenderTargetTextureAdapter->OverrideWithTextureFormat = Texture2D;
-					TextureData->PaintRenderTargetTextureAdapter->bUseDefaultTileSizes = false;
-					TextureData->PaintRenderTargetTextureAdapter->TileSize = VirtualTextureBuildSettings.TileSize;
-					TextureData->PaintRenderTargetTextureAdapter->TileBorderSize = VirtualTextureBuildSettings.TileBorderSize;
-					TextureData->PaintRenderTargetTextureAdapter->UpdateResource();
+					PaintableMeshComponents.Add(MeshComponent);
 				}
 			}
+		}
 
-			TextureData->PaintRenderTargetTexture->AddressX = Texture2D->AddressX;
-			TextureData->PaintRenderTargetTexture->AddressY = Texture2D->AddressY;
+		if (PaintableMeshComponents.Num() == 0)
+		{
+			continue;
+		}
 
-			// Initialize the render target with the texture contents.
-			UTexturePaintToolset::SetupInitialRenderTargetData(TextureData->PaintingTexture2D, TextureData->PaintRenderTargetTexture);
+		UTexture2D* Texture2D = Cast<UTexture2D>(PaintableTexture.Texture);
+		if (Texture2D == nullptr)
+		{
+			continue;
+		}
+	
+		Texture2D->BlockOnAnyAsyncBuild();
 
-			// Need to flush the virtual texture adapter since we just updated the painting render target.
-			if (TextureData->PaintRenderTargetTextureAdapter)
-			{
-				TextureData->PaintRenderTargetTextureAdapter->Flush(FBox2f(FVector2f(0, 0), FVector2f(1, 1)));
-			}
+		// Create the render target to paint on.
+		FPaintTexture2DData* TextureData = AddPaintTargetData(Texture2D);
 
-			// Apply the render target override only to the components that we are painting with this texture.
-			UMeshPaintingSubsystem* MeshPaintingSubsystem = GEngine->GetEngineSubsystem<UMeshPaintingSubsystem>();
-			TArray<UMeshComponent*> SelectedMeshComponents = MeshPaintingSubsystem->GetSelectedMeshComponents();
-			for (int32 Index = 0; Index < SelectedMeshComponents.Num(); ++Index)
+		const int32 TextureWidth = Texture2D->Source.GetSizeX();
+		const int32 TextureHeight = Texture2D->Source.GetSizeY();
+
+		if (TextureData->PaintRenderTargetTexture == nullptr ||
+			TextureData->PaintRenderTargetTexture->GetSurfaceWidth() != TextureWidth ||
+			TextureData->PaintRenderTargetTexture->GetSurfaceHeight() != TextureHeight)
+		{
+			TextureData->PaintRenderTargetTexture = NewObject<UTextureRenderTarget2D>(GetTransientPackage(), NAME_None, RF_Transient);
+			TextureData->PaintRenderTargetTexture->bNeedsTwoCopies = true;
+			const bool bForceLinearGamma = true;
+			TextureData->PaintRenderTargetTexture->InitCustomFormat(TextureWidth, TextureHeight, PF_A16B16G16R16, bForceLinearGamma);
+			TextureData->PaintRenderTargetTexture->UpdateResourceImmediate();
+
+			TextureData->PaintRenderTargetTextureAdapter = nullptr; 
+			if (TextureData->PaintingTexture2D->IsCurrentlyVirtualTextured())
 			{
-				UMeshComponent* MeshComponent = SelectedMeshComponents[Index];
-				if (CanPaintTextureToComponent(PaintableTexture.Texture, MeshComponent))
-				{
-					TSharedPtr<IMeshPaintComponentAdapter> MeshAdapter = MeshPaintingSubsystem->GetAdapterForComponent(MeshComponent);
-					if (IsMeshAdapterSupported(MeshAdapter))
-					{
-						AddTextureOverrideToComponent(*TextureData, MeshComponent, MeshAdapter.Get());
-					}
-				}
+				// Virtual textures can't just swap in a render target in their material, so we use a virtual texture adapter.
+				FVirtualTextureBuildSettings VirtualTextureBuildSettings;
+				TextureData->PaintingTexture2D->GetVirtualTextureBuildSettings(VirtualTextureBuildSettings);
+
+				TextureData->PaintRenderTargetTextureAdapter = NewObject<UVirtualTextureAdapter>(GetTransientPackage(), NAME_None, RF_Transient);
+				TextureData->PaintRenderTargetTextureAdapter->Texture = TextureData->PaintRenderTargetTexture;
+				TextureData->PaintRenderTargetTextureAdapter->OverrideWithTextureFormat = Texture2D;
+				TextureData->PaintRenderTargetTextureAdapter->bUseDefaultTileSizes = false;
+				TextureData->PaintRenderTargetTextureAdapter->TileSize = VirtualTextureBuildSettings.TileSize;
+				TextureData->PaintRenderTargetTextureAdapter->TileBorderSize = VirtualTextureBuildSettings.TileBorderSize;
+				TextureData->PaintRenderTargetTextureAdapter->UpdateResource();
 			}
 		}
+
+		TextureData->PaintRenderTargetTexture->AddressX = Texture2D->AddressX;
+		TextureData->PaintRenderTargetTexture->AddressY = Texture2D->AddressY;
+
+		// Initialize the render target with the texture contents.
+		UTexturePaintToolset::SetupInitialRenderTargetData(TextureData->PaintingTexture2D, TextureData->PaintRenderTargetTexture);
+
+		// Need to flush the virtual texture adapter since we just updated the painting render target.
+		if (TextureData->PaintRenderTargetTextureAdapter)
+		{
+			TextureData->PaintRenderTargetTextureAdapter->Flush(FBox2f(FVector2f(0, 0), FVector2f(1, 1)));
+		}
+
+		// Apply the overrides.
+		for (UMeshComponent* MeshComponent : PaintableMeshComponents)
+		{
+			TSharedPtr<IMeshPaintComponentAdapter> MeshAdapter = MeshPaintingSubsystem->GetAdapterForComponent(MeshComponent);
+			AddTextureOverrideToComponent(*TextureData, MeshComponent, MeshAdapter.Get());
+		}
 	}
 }
 
