MeshPaint: Add GetNumUVChannels() virtual function to the component adapter.
This will be needed to correctly validate if we can paint on a given channel.

diff --git a/Engine/Plugins/MeshPainting/Source/MeshPaintingToolset/Private/MeshPaintGeometryCollectionAdapter.cpp b/Engine/Plugins/MeshPainting/Source/MeshPaintingToolset/Private/MeshPaintGeometryCollectionAdapter.cpp
--- a/Engine/Plugins/MeshPainting/Source/MeshPaintingToolset/Private/MeshPaintGeometryCollectionAdapter.cpp
+++ b/Engine/Plugins/MeshPainting/Source/MeshPaintingToolset/Private/MeshPaintGeometryCollectionAdapter.cpp
@@ -158,6 +158,15 @@ void FMeshPaintGeometryCollectionComponentAdapter::OnRemoved()
 	}
 }
 
+int32 FMeshPaintGeometryCollectionComponentAdapter::GetNumUVChannels() const
+{
+	if (!IsValid())
+	{
+		return 0;
+	}
+	return GetGeometryCollectionObject()->GetGeometryCollection()->NumUVLayers();
+}
+
 bool FMeshPaintGeometryCollectionComponentAdapter::LineTraceComponent(struct FHitResult& OutHit, const FVector Start, const FVector End, const struct FCollisionQueryParams& Params) const
 {
 	// Note: This is 99% the same as the skeletal mesh component adapter version, and could likely be pushed up to the base class if we added a generalized way to get the Component
diff --git a/Engine/Plugins/MeshPainting/Source/MeshPaintingToolset/Private/MeshPaintHelpers.cpp b/Engine/Plugins/MeshPainting/Source/MeshPaintingToolset/Private/MeshPaintHelpers.cpp
--- a/Engine/Plugins/MeshPainting/Source/MeshPaintingToolset/Private/MeshPaintHelpers.cpp
+++ b/Engine/Plugins/MeshPainting/Source/MeshPaintingToolset/Private/MeshPaintHelpers.cpp
@@ -1098,7 +1098,7 @@ int32 UMeshPaintingSubsystem::GetNumberOfLODs(const UMeshComponent* MeshComponen
 	return NumLODs;
 }
 
-int32 UMeshPaintingSubsystem::GetNumberOfUVs(const UMeshComponent* MeshComponent, int32 LODIndex)
+int32 UMeshPaintingSubsystem::GetNumberOfUVs(const UMeshComponent* MeshComponent, int32 LODIndex) const
 {
 	int32 NumUVs = 0;
 
diff --git a/Engine/Plugins/MeshPainting/Source/MeshPaintingToolset/Private/MeshPaintSkeletalMeshAdapter.cpp b/Engine/Plugins/MeshPainting/Source/MeshPaintingToolset/Private/MeshPaintSkeletalMeshAdapter.cpp
--- a/Engine/Plugins/MeshPainting/Source/MeshPaintingToolset/Private/MeshPaintSkeletalMeshAdapter.cpp
+++ b/Engine/Plugins/MeshPainting/Source/MeshPaintingToolset/Private/MeshPaintSkeletalMeshAdapter.cpp
@@ -264,6 +264,11 @@ void FMeshPaintSkeletalMeshComponentAdapter::OnRemoved()
 	ReferencedSkeletalMesh->OnPostMeshCached().RemoveAll(this);
 }
 
+int32 FMeshPaintSkeletalMeshComponentAdapter::GetNumUVChannels() const
+{
+	return LODData == nullptr ? 0 : LODData->GetNumTexCoords();
+}
+
 bool FMeshPaintSkeletalMeshComponentAdapter::LineTraceComponent(struct FHitResult& OutHit, const FVector Start, const FVector End, const struct FCollisionQueryParams& Params) const
 {
 	if (!SkeletalMeshComponent.IsValid())
diff --git a/Engine/Plugins/MeshPainting/Source/MeshPaintingToolset/Private/MeshPaintStaticMeshAdapter.cpp b/Engine/Plugins/MeshPainting/Source/MeshPaintingToolset/Private/MeshPaintStaticMeshAdapter.cpp
--- a/Engine/Plugins/MeshPainting/Source/MeshPaintingToolset/Private/MeshPaintStaticMeshAdapter.cpp
+++ b/Engine/Plugins/MeshPainting/Source/MeshPaintingToolset/Private/MeshPaintStaticMeshAdapter.cpp
@@ -170,6 +170,10 @@ void FMeshPaintStaticMeshComponentAdapter::AddReferencedObjects(FReferenceCollec
 	Collector.AddReferencedObject(ReferencedStaticMesh);
 }
 
+int32 FMeshPaintStaticMeshComponentAdapter::GetNumUVChannels() const
+{
+	return LODModel == nullptr ? 0 : LODModel->GetNumTexCoords();
+}
 
 bool FMeshPaintStaticMeshComponentAdapter::LineTraceComponent(struct FHitResult& OutHit, const FVector Start, const FVector End, const struct FCollisionQueryParams& Params) const
 {
diff --git a/Engine/Plugins/MeshPainting/Source/MeshPaintingToolset/Public/IMeshPaintComponentAdapter.h b/Engine/Plugins/MeshPainting/Source/MeshPaintingToolset/Public/IMeshPaintComponentAdapter.h
--- a/Engine/Plugins/MeshPainting/Source/MeshPaintingToolset/Public/IMeshPaintComponentAdapter.h
+++ b/Engine/Plugins/MeshPainting/Source/MeshPaintingToolset/Public/IMeshPaintComponentAdapter.h
@@ -52,6 +52,9 @@ public:
 	/** Whether or not this adapter supports vertex painting */
 	virtual bool SupportsVertexPaint() const = 0;
 
+	/** Returns the number of UV channels for the current LOD index in the Mesh. */
+	virtual int32 GetNumUVChannels() const = 0;
+
 	/** Returns the result of tracing a line on the component represented by this adapter */
 	virtual bool LineTraceComponent(struct FHitResult& OutHit, const FVector Start, const FVector End, const struct FCollisionQueryParams& Params) const = 0;	
 	
diff --git a/Engine/Plugins/MeshPainting/Source/MeshPaintingToolset/Public/MeshPaintGeometryCollectionAdapter.h b/Engine/Plugins/MeshPainting/Source/MeshPaintingToolset/Public/MeshPaintGeometryCollectionAdapter.h
--- a/Engine/Plugins/MeshPainting/Source/MeshPaintingToolset/Public/MeshPaintGeometryCollectionAdapter.h
+++ b/Engine/Plugins/MeshPainting/Source/MeshPaintingToolset/Public/MeshPaintGeometryCollectionAdapter.h
@@ -31,6 +31,7 @@ public:
 	virtual bool SupportsTexturePaint() const override { return IsValid(); }
 	virtual bool SupportsTextureColorPaint() const override { return false; }
 	virtual bool SupportsVertexPaint() const override { return IsValid(); }
+	virtual int32 GetNumUVChannels() const override;
 	virtual bool LineTraceComponent(struct FHitResult& OutHit, const FVector Start, const FVector End, const struct FCollisionQueryParams& Params) const override;
 	virtual void QueryPaintableTextures(int32 MaterialIndex, int32& OutDefaultIndex, TArray<struct FPaintableTexture>& InOutTextureList) override;
 	virtual void ApplyOrRemoveTextureOverride(UTexture* SourceTexture, UTexture* OverrideTexture) const override;
diff --git a/Engine/Plugins/MeshPainting/Source/MeshPaintingToolset/Public/MeshPaintHelpers.h b/Engine/Plugins/MeshPainting/Source/MeshPaintingToolset/Public/MeshPaintHelpers.h
--- a/Engine/Plugins/MeshPainting/Source/MeshPaintingToolset/Public/MeshPaintHelpers.h
+++ b/Engine/Plugins/MeshPainting/Source/MeshPaintingToolset/Public/MeshPaintHelpers.h
@@ -215,7 +215,7 @@ public:
 	bool TryGetNumberOfLODs(const UMeshComponent* MeshComponent, int32& OutNumLODs);
 	
 	/** Returns the number of Texture Coordinates for the given MeshComponent */
-	int32 GetNumberOfUVs(const UMeshComponent* MeshComponent, int32 LODIndex);
+	int32 GetNumberOfUVs(const UMeshComponent* MeshComponent, int32 LODIndex) const;
 
 	/** Checks whether or not the mesh components contains per lod colors (for all LODs)*/
 	bool DoesMeshComponentContainPerLODColors(const UMeshComponent* MeshComponent);
diff --git a/Engine/Plugins/MeshPainting/Source/MeshPaintingToolset/Public/MeshPaintSkeletalMeshAdapter.h b/Engine/Plugins/MeshPainting/Source/MeshPaintingToolset/Public/MeshPaintSkeletalMeshAdapter.h
--- a/Engine/Plugins/MeshPainting/Source/MeshPaintingToolset/Public/MeshPaintSkeletalMeshAdapter.h
+++ b/Engine/Plugins/MeshPainting/Source/MeshPaintingToolset/Public/MeshPaintSkeletalMeshAdapter.h
@@ -33,6 +33,7 @@ public:
 	virtual bool SupportsTexturePaint() const override { return true; }
 	virtual bool SupportsTextureColorPaint() const override { return false; }
 	virtual bool SupportsVertexPaint() const override { return SkeletalMeshComponent.IsValid(); }
+	virtual int32 GetNumUVChannels() const override;
 	virtual bool LineTraceComponent(struct FHitResult& OutHit, const FVector Start, const FVector End, const struct FCollisionQueryParams& Params) const override;
 	virtual void QueryPaintableTextures(int32 MaterialIndex, int32& OutDefaultIndex, TArray<struct FPaintableTexture>& InOutTextureList) override;
 	virtual void ApplyOrRemoveTextureOverride(UTexture* SourceTexture, UTexture* OverrideTexture) const override;
diff --git a/Engine/Plugins/MeshPainting/Source/MeshPaintingToolset/Public/MeshPaintStaticMeshAdapter.h b/Engine/Plugins/MeshPainting/Source/MeshPaintingToolset/Public/MeshPaintStaticMeshAdapter.h
--- a/Engine/Plugins/MeshPainting/Source/MeshPaintingToolset/Public/MeshPaintStaticMeshAdapter.h
+++ b/Engine/Plugins/MeshPainting/Source/MeshPaintingToolset/Public/MeshPaintStaticMeshAdapter.h
@@ -29,6 +29,7 @@ public:
 	virtual bool IsValid() const override { return StaticMeshComponent.IsValid() && ReferencedStaticMesh && StaticMeshComponent->GetStaticMesh() == ReferencedStaticMesh; }
 	virtual bool SupportsTexturePaint() const override { return true; }
 	virtual bool SupportsTextureColorPaint() const override { return StaticMeshComponent.IsValid() && StaticMeshComponent->CanMeshPaintTextureColors(); }
+	virtual int32 GetNumUVChannels() const override;
 	virtual bool SupportsVertexPaint() const override { return StaticMeshComponent.IsValid() && StaticMeshComponent->CanMeshPaintVertexColors(); }
 	virtual bool LineTraceComponent(struct FHitResult& OutHit, const FVector Start, const FVector End, const struct FCollisionQueryParams& Params) const override;	
 	virtual void QueryPaintableTextures(int32 MaterialIndex, int32& OutDefaultIndex, TArray<struct FPaintableTexture>& InOutTextureList) override;
