MeshPaint: Move some logic for rejecting textures for painting to earlier in the discovery.
This avoids bug where we remove from the found texture array after having stored the index in the array of the default texture.

diff --git a/Engine/Plugins/MeshPainting/Source/MeshPaintingToolset/Private/IMeshPaintComponentAdapter.cpp b/Engine/Plugins/MeshPainting/Source/MeshPaintingToolset/Private/IMeshPaintComponentAdapter.cpp
--- a/Engine/Plugins/MeshPainting/Source/MeshPaintingToolset/Private/IMeshPaintComponentAdapter.cpp
+++ b/Engine/Plugins/MeshPainting/Source/MeshPaintingToolset/Private/IMeshPaintComponentAdapter.cpp
@@ -47,6 +47,7 @@ static bool IsTextureSuitableForTexturePainting(const TWeakObjectPtr<UTexture> T
 		!TexturePtr->IsNormalMap() &&
 		!TexturePtr->HasHDRSource() && // Currently HDR textures are not supported to paint on.
 		TexturePtr->Source.IsValid() &&
+		TexturePtr->Source.GetFormat() != TSF_G16 && // Currently 16 bit textures are not supported to paint on.
 		TexturePtr->Source.GetBytesPerPixel() > 0 && // Textures' sources must have a known count of bytes per pixel
 		(TexturePtr->Source.GetBytesPerPixel() <= UTexturePaintToolset::GetMaxSupportedBytesPerPixelForPainting())); // Textures' sources must fit in FColor struct to be supported.
 }
diff --git a/Engine/Plugins/MeshPainting/Source/MeshPaintingToolset/Private/TexturePaintToolset.cpp b/Engine/Plugins/MeshPainting/Source/MeshPaintingToolset/Private/TexturePaintToolset.cpp
--- a/Engine/Plugins/MeshPainting/Source/MeshPaintingToolset/Private/TexturePaintToolset.cpp
+++ b/Engine/Plugins/MeshPainting/Source/MeshPaintingToolset/Private/TexturePaintToolset.cpp
@@ -473,17 +473,5 @@ void UTexturePaintToolset::RetrieveTexturesForComponent(const UMeshComponent* Co
 				OutDefaultIndex = OutDefaultIndexForMaterial;
 			}
 		}
-
-		// Filter out any paintable texture that is not supported currently (Todo: should we do this here or could we just move some logic in the FPaintableTexture)?
-		OutTextures.RemoveAll([](const FPaintableTexture& PaintableTexture)
-		{
-			if (const UTexture* Texture = PaintableTexture.Texture)
-			{
-				// Only 8 bit images are supported by the tool currently (we should probably provide some feedback to the user on why some texture are not supported and a way to convert them)
-				return Texture->Source.GetFormat() != TSF_BGRA8 && Texture->Source.GetFormat() != TSF_G8;
-			}
-
-			return true;
-		});
 	}
 }
