MeshPaint: Test for valid UV channel in target mesh before doing any painting.
This fixes a crash when painting with the wrong UV channel.

diff --git a/Engine/Plugins/MeshPainting/Source/MeshPaintingToolset/Private/MeshTexturePaintingTool.cpp b/Engine/Plugins/MeshPainting/Source/MeshPaintingToolset/Private/MeshTexturePaintingTool.cpp
--- a/Engine/Plugins/MeshPainting/Source/MeshPaintingToolset/Private/MeshTexturePaintingTool.cpp
+++ b/Engine/Plugins/MeshPainting/Source/MeshPaintingToolset/Private/MeshTexturePaintingTool.cpp
@@ -187,59 +187,66 @@ void UMeshTexturePaintingTool::OnTick(float DeltaTime)
 			for (int32 j = 0; j < SelectedMeshComponents.Num(); ++j)
 			{
 				UMeshComponent* SelectedComponent = SelectedMeshComponents[j];
-				
+				TSharedPtr<IMeshPaintComponentAdapter> MeshAdapter = MeshPaintingSubsystem->GetAdapterForComponent(SelectedComponent);
+				if (!IsMeshAdapterSupported(MeshAdapter))
+				{
+					continue;
+				}
+
 				const int32 UVChannel = GetSelectedUVChannel(SelectedComponent);
+				if (UVChannel >= MeshAdapter->GetNumUVChannels())
+				{
+					continue;
+				}
 
-				TSharedPtr<IMeshPaintComponentAdapter> MeshAdapter = MeshPaintingSubsystem->GetAdapterForComponent(SelectedComponent);
-				if (IsMeshAdapterSupported(MeshAdapter))
+				TArray<UTexture const*> Textures;
+				const UTexture2D* TargetTexture2D = GetSelectedPaintTexture(SelectedComponent);
+				if (TargetTexture2D == nullptr)
 				{
-					TArray<UTexture const*> Textures;
-					const UTexture2D* TargetTexture2D = GetSelectedPaintTexture(SelectedComponent);
-					if (TargetTexture2D)
-					{
-						Textures.Add(TargetTexture2D);
-
-						FPaintTexture2DData* TextureData = GetPaintTargetData(TargetTexture2D);
-						if (TextureData)
-						{
-							Textures.Add(TextureData->PaintRenderTargetTexture);
-						}
-
-						TArray<FTexturePaintMeshSectionInfo> MaterialSections;
-						UTexturePaintToolset::RetrieveMeshSectionsForTextures(SelectedComponent, 0 /*CachedLODIndex*/, Textures, MaterialSections);
-
-						TArray<FTexturePaintTriangleInfo> TrianglePaintInfoArray;
-						FPerTrianglePaintAction TempAction = FPerTrianglePaintAction::CreateUObject(this, &UMeshTexturePaintingTool::GatherTextureTriangles, &TrianglePaintInfoArray, &MaterialSections, UVChannel);
-
-						// We are flooding the texture, so all triangles are influenced
-
-						const TArray<uint32>& MeshIndices = MeshAdapter->GetMeshIndices();
-						int32 TriangleIndices[3];
-
-						for (int32 i = 0; i < MeshIndices.Num(); i += 3)
-						{
-							TriangleIndices[0] = MeshIndices[i + 0];
-							TriangleIndices[1] = MeshIndices[i + 1];
-							TriangleIndices[2] = MeshIndices[i + 2];
-							TempAction.Execute(MeshAdapter.Get(), i / 3, TriangleIndices);
-						}
-
-						// Painting textures
-						if ((TexturePaintingCurrentMeshComponent != nullptr) && (TexturePaintingCurrentMeshComponent != SelectedComponent))
-						{
-							// Mesh has changed, so finish up with our previous texture
-							FinishPaintingTexture();
-						}
-
-						if (TexturePaintingCurrentMeshComponent == nullptr)
-						{
-							StartPaintingTexture(SelectedComponent, *MeshAdapter);
-						}
-
-						FMeshPaintParameters* LastParams = nullptr;
-						PaintTexture(BucketFillParams, UVChannel, TrianglePaintInfoArray, *MeshAdapter, LastParams);
-					}
+					continue;
+				}
+					
+				Textures.Add(TargetTexture2D);
+
+				FPaintTexture2DData* TextureData = GetPaintTargetData(TargetTexture2D);
+				if (TextureData)
+				{
+					Textures.Add(TextureData->PaintRenderTargetTexture);
+				}
+
+				TArray<FTexturePaintMeshSectionInfo> MaterialSections;
+				UTexturePaintToolset::RetrieveMeshSectionsForTextures(SelectedComponent, 0 /*CachedLODIndex*/, Textures, MaterialSections);
+
+				TArray<FTexturePaintTriangleInfo> TrianglePaintInfoArray;
+				FPerTrianglePaintAction TempAction = FPerTrianglePaintAction::CreateUObject(this, &UMeshTexturePaintingTool::GatherTextureTriangles, &TrianglePaintInfoArray, &MaterialSections, UVChannel);
+
+				// We are flooding the texture, so all triangles are influenced
+
+				const TArray<uint32>& MeshIndices = MeshAdapter->GetMeshIndices();
+				int32 TriangleIndices[3];
+
+				for (int32 i = 0; i < MeshIndices.Num(); i += 3)
+				{
+					TriangleIndices[0] = MeshIndices[i + 0];
+					TriangleIndices[1] = MeshIndices[i + 1];
+					TriangleIndices[2] = MeshIndices[i + 2];
+					TempAction.Execute(MeshAdapter.Get(), i / 3, TriangleIndices);
 				}
+
+				// Painting textures
+				if ((TexturePaintingCurrentMeshComponent != nullptr) && (TexturePaintingCurrentMeshComponent != SelectedComponent))
+				{
+					// Mesh has changed, so finish up with our previous texture
+					FinishPaintingTexture();
+				}
+
+				if (TexturePaintingCurrentMeshComponent == nullptr)
+				{
+					StartPaintingTexture(SelectedComponent, *MeshAdapter);
+				}
+
+				FMeshPaintParameters* LastParams = nullptr;
+				PaintTexture(BucketFillParams, UVChannel, TrianglePaintInfoArray, *MeshAdapter, LastParams);
 			}
 		}
 
@@ -472,62 +479,70 @@ bool UMeshTexturePaintingTool::PaintInternal(const TArrayView<TPair<FVector, FVe
 			TArray<int32>& PaintRayResultIds = Entry.Value;
 			
 			TSharedPtr<IMeshPaintComponentAdapter> MeshAdapter = MeshPaintingSubsystem->GetAdapterForComponent(HoveredComponent);
-			if (IsMeshAdapterSupported(MeshAdapter))
+			if (!IsMeshAdapterSupported(MeshAdapter))
+			{
+				continue;
+			}
+				
+			const int32 UVChannel = GetSelectedUVChannel(HoveredComponent);
+			if (UVChannel >= MeshAdapter->GetNumUVChannels())
+			{
+				continue;
+			}
+
+			TArray<UTexture const*> Textures;
+			const UTexture2D* TargetTexture2D = GetSelectedPaintTexture(HoveredComponent);
+			if (TargetTexture2D == nullptr)
 			{
-				const int32 UVChannel = GetSelectedUVChannel(HoveredComponent);
+				continue;
+			}
+					
+			Textures.Add(TargetTexture2D);
 
-				TArray<UTexture const*> Textures;
-				const UTexture2D* TargetTexture2D = GetSelectedPaintTexture(HoveredComponent);
-				if (TargetTexture2D)
-				{
-					Textures.Add(TargetTexture2D);
+			FPaintTexture2DData* TextureData = GetPaintTargetData(TargetTexture2D);
+			if (TextureData)
+			{
+				Textures.Add(TextureData->PaintRenderTargetTexture);
+			}
 
-					FPaintTexture2DData* TextureData = GetPaintTargetData(TargetTexture2D);
-					if (TextureData)
-					{
-						Textures.Add(TextureData->PaintRenderTargetTexture);
-					}
+			TArray<FTexturePaintMeshSectionInfo> MaterialSections;
+			UTexturePaintToolset::RetrieveMeshSectionsForTextures(HoveredComponent, 0/*CachedLODIndex*/, Textures, MaterialSections);
 
-					TArray<FTexturePaintMeshSectionInfo> MaterialSections;
-					UTexturePaintToolset::RetrieveMeshSectionsForTextures(HoveredComponent, 0/*CachedLODIndex*/, Textures, MaterialSections);
+			TArray<FTexturePaintTriangleInfo> TrianglePaintInfoArray;
+			for (int32 PaintRayResultId : PaintRayResultIds)
+			{
+				const FVector& BestTraceResultLocation = PaintRayResults[PaintRayResultId].BestTraceResult.Location;
+				FViewCameraState CameraState;
+				GetToolManager()->GetContextQueriesAPI()->GetCurrentViewState(CameraState);
+				bPaintApplied |= MeshPaintingSubsystem->ApplyPerTrianglePaintAction(MeshAdapter.Get(), CameraState.Position, BestTraceResultLocation, BrushProperties, FPerTrianglePaintAction::CreateUObject(this, &UMeshTexturePaintingTool::GatherTextureTriangles, &TrianglePaintInfoArray, &MaterialSections, UVChannel), TextureProperties->bOnlyFrontFacingTriangles);
+				break;
+			}
 
-					TArray<FTexturePaintTriangleInfo> TrianglePaintInfoArray;
-					for (int32 PaintRayResultId : PaintRayResultIds)
-					{
-						const FVector& BestTraceResultLocation = PaintRayResults[PaintRayResultId].BestTraceResult.Location;
-						FViewCameraState CameraState;
-						GetToolManager()->GetContextQueriesAPI()->GetCurrentViewState(CameraState);
-						bPaintApplied |= MeshPaintingSubsystem->ApplyPerTrianglePaintAction(MeshAdapter.Get(), CameraState.Position, BestTraceResultLocation, BrushProperties, FPerTrianglePaintAction::CreateUObject(this, &UMeshTexturePaintingTool::GatherTextureTriangles, &TrianglePaintInfoArray, &MaterialSections, UVChannel), TextureProperties->bOnlyFrontFacingTriangles);
-						break;
-					}
+			// Painting textures
+			if ((TexturePaintingCurrentMeshComponent != nullptr) && (TexturePaintingCurrentMeshComponent != HoveredComponent))
+			{
+				// Mesh has changed, so finish up with our previous texture
+				FinishPaintingTexture();
+			}
 
-					// Painting textures
-					if ((TexturePaintingCurrentMeshComponent != nullptr) && (TexturePaintingCurrentMeshComponent != HoveredComponent))
-					{
-						// Mesh has changed, so finish up with our previous texture
-						FinishPaintingTexture();
-					}
+			if (TexturePaintingCurrentMeshComponent == nullptr)
+			{
+				StartPaintingTexture(HoveredComponent, *MeshAdapter);
+			}
 
-					if (TexturePaintingCurrentMeshComponent == nullptr)
+			if (TexturePaintingCurrentMeshComponent != nullptr)
+			{
+				for (int32 PaintRayResultId : PaintRayResultIds)
+				{
+					FMeshPaintParameters& Params = PaintRayResults[PaintRayResultId].Params;
+					FMeshPaintParameters* LastParams = nullptr;
+					if (LastPaintRayResults.Num() > PaintRayResultId)
 					{
-						StartPaintingTexture(HoveredComponent, *MeshAdapter);
+						LastParams = &LastPaintRayResults[PaintRayResultId].Params;
 					}
 
-					if (TexturePaintingCurrentMeshComponent != nullptr)
-					{
-						for (int32 PaintRayResultId : PaintRayResultIds)
-						{
-							FMeshPaintParameters& Params = PaintRayResults[PaintRayResultId].Params;
-							FMeshPaintParameters* LastParams = nullptr;
-							if (LastPaintRayResults.Num() > PaintRayResultId)
-							{
-								LastParams = &LastPaintRayResults[PaintRayResultId].Params;
-							}
-
-							PaintTexture(Params, UVChannel, TrianglePaintInfoArray, *MeshAdapter, LastParams);
-							break;
-						}
-					}
+					PaintTexture(Params, UVChannel, TrianglePaintInfoArray, *MeshAdapter, LastParams);
+					break;
 				}
 			}
 		}
@@ -1511,10 +1526,10 @@ void UMeshTextureAssetPaintingTool::CacheTexturePaintData()
 	if (MeshPaintingSubsystem)
 	{
 		PaintableTextures.Empty();
-		
+
 		UTexture* DefaultTexture = nullptr;
 		int32 DefaultUVChannelIndex = INDEX_NONE;
-	
+
 		TArray<UMeshComponent*> PaintableComponents = MeshPaintingSubsystem->GetPaintableMeshComponents();
 		if (PaintableComponents.IsValidIndex(0))
 		{
@@ -1527,7 +1542,7 @@ void UMeshTextureAssetPaintingTool::CacheTexturePaintData()
 				DefaultUVChannelIndex = PaintableTextures[DefaultTextureIndex].UVChannelIndex;
 			}
 		}
-
+		
 		PaintableTextures.RemoveAll([](FPaintableTexture const& PaintableTexture) { return PaintableTexture.bIsMeshTexture; });
 
 		// Ensure that the selection remains valid or is invalidated
@@ -1539,7 +1554,7 @@ void UMeshTextureAssetPaintingTool::CacheTexturePaintData()
 			{
 				NewTexture = Cast<UTexture2D>(DefaultTexture);
 				NewUVChannelIndex = DefaultUVChannelIndex;
-			}
+		}
 			else if (PaintableTextures.Num() > 0)
 			{
 				NewTexture = Cast<UTexture2D>(PaintableTextures[0].Texture);
