Support multiselection of components when in mesh texture asset paint mode.
When multiselecting the list of paintable textures is the intersection of the per component texture lists.
With multiselecting it is possible to select a UV channel to paint on one component that is not available on another.
For this case we don't paint to the components with UV channel out of range, but we allow them to stay selected.
A warning icon and tool text is displayed in the UV channel select UI.

diff --git a/Engine/Plugins/MeshPainting/Source/MeshPaintEditorMode/Private/BrushSettingsCustomization.cpp b/Engine/Plugins/MeshPainting/Source/MeshPaintEditorMode/Private/BrushSettingsCustomization.cpp
--- a/Engine/Plugins/MeshPainting/Source/MeshPaintEditorMode/Private/BrushSettingsCustomization.cpp
+++ b/Engine/Plugins/MeshPainting/Source/MeshPaintEditorMode/Private/BrushSettingsCustomization.cpp
@@ -502,21 +502,47 @@ void FTextureAssetPaintingSettingsCustomization::CustomizeDetails(IDetailLayoutB
 			SNew(SNumericEntryBox<int32>)
 			.Font(IDetailLayoutBuilder::GetDetailFont())
 			.AllowSpin(true)
-			.Value_Lambda([]() -> int32 { return UMeshPaintMode::GetTextureAssetToolProperties() ? UMeshPaintMode::GetTextureAssetToolProperties()->UVChannel : 0; })
+			.Value_Lambda([]() -> int32 
+			{ 
+				return UMeshPaintMode::GetTextureAssetToolProperties() ? UMeshPaintMode::GetTextureAssetToolProperties()->UVChannel : 0; 
+			})
 			.MinValue(0)
-			.MaxValue_Lambda([]() -> int32 { return GEngine->GetEngineSubsystem<UMeshPaintingSubsystem>()->GetMaxUVIndexToPaint(); })
-			.OnValueChanged(SNumericEntryBox<int32>::FOnValueChanged::CreateLambda([=](int32 Value) {
-			if (UMeshPaintMode::GetTextureAssetToolProperties())
+			.MaxValue_Lambda([]() -> int32 
 			{
-				UMeshPaintMode::GetTextureAssetToolProperties()->UVChannel = Value;
-			}
-				}))
-			.OnValueCommitted(SNumericEntryBox<int32>::FOnValueCommitted::CreateLambda([](int32 Value, ETextCommit::Type CommitType) {
-			if (UMeshPaintMode::GetTextureAssetToolProperties())
+				return FMath::Max(0, GEngine->GetEngineSubsystem<UMeshPaintingSubsystem>()->GetMinMaxUVChannelsToPaint()[1] - 1);
+			})
+			.OnValueChanged(SNumericEntryBox<int32>::FOnValueChanged::CreateLambda([=](int32 Value)
 			{
-				UMeshPaintMode::GetTextureAssetToolProperties()->UVChannel = Value;
-			}
-		}))
+				if (UMeshPaintMode::GetTextureAssetToolProperties())
+				{
+					UMeshPaintMode::GetTextureAssetToolProperties()->UVChannel = Value;
+				}
+			}))
+			.OnValueCommitted(SNumericEntryBox<int32>::FOnValueCommitted::CreateLambda([](int32 Value, ETextCommit::Type CommitType)
+			{
+				if (UMeshPaintMode::GetTextureAssetToolProperties())
+				{
+					UMeshPaintMode::GetTextureAssetToolProperties()->UVChannel = Value;
+				}
+			}))
+		]
+		+ SHorizontalBox::Slot()
+		.HAlign(HAlign_Left)
+		.VAlign(VAlign_Center)
+		[
+			SNew(SImage)
+			.Image(FCoreStyle::Get().GetBrush("Icons.Warning"))
+			.Visibility_Lambda([this]()
+			{
+				const int32 UVChannel = UMeshPaintMode::GetTextureAssetToolProperties() ? UMeshPaintMode::GetTextureAssetToolProperties()->UVChannel : 0;
+				const int32 NumSafeUVChannels = GEngine->GetEngineSubsystem<UMeshPaintingSubsystem>()->GetMinMaxUVChannelsToPaint()[0];
+				const bool bIsSafe = UVChannel < NumSafeUVChannels;
+				return bIsSafe ? EVisibility::Hidden : EVisibility::Visible;
+			})
+			.ToolTipText_Lambda([this]() 
+			{
+				return LOCTEXT("UVChannelWarning", "Some selected components do not support the selected UV channel.\nThey will not receive painting.");
+			})
 		]
 	];
 
diff --git a/Engine/Plugins/MeshPainting/Source/MeshPaintingToolset/Private/MeshPaintHelpers.cpp b/Engine/Plugins/MeshPainting/Source/MeshPaintingToolset/Private/MeshPaintHelpers.cpp
--- a/Engine/Plugins/MeshPainting/Source/MeshPaintingToolset/Private/MeshPaintHelpers.cpp
+++ b/Engine/Plugins/MeshPainting/Source/MeshPaintingToolset/Private/MeshPaintHelpers.cpp
@@ -1614,6 +1614,9 @@ void UMeshPaintingSubsystem::UpdatePaintSupportState()
 	bSelectionSupportsTextureColorPaint = false;
 	bSelectionSupportsTextureAssetPaint = false;
 
+	TArray<FPaintableTexture> CommonPaintableTextures;
+	bool bCommonPaintableTexturesInitialized = false;
+
 	for (TWeakObjectPtr<UMeshComponent> MeshComponentWeak : SelectedMeshComponents)
 	{
 		if (UMeshComponent* MeshComponent = MeshComponentWeak.Get())
@@ -1648,13 +1651,36 @@ void UMeshPaintingSubsystem::UpdatePaintSupportState()
 						bSelectionSupportsTextureColorPaint |= bSupportsTextureColorPaint && PaintableTexture.bIsMeshTexture;
 						bSelectionSupportsTextureAssetPaint |= bSupportsTextureAssetPaint && !PaintableTexture.bIsMeshTexture;
 					}
+
+					if (bSupportsTextureAssetPaint)
+					{
+						// Fill CommonPaintableTextures array with textures that are in present in ALL components.
+						if (!bCommonPaintableTexturesInitialized)
+						{
+							CommonPaintableTextures = MoveTemp(PaintableTextures);
+							bCommonPaintableTexturesInitialized = true;
+						}
+						else
+						{
+							TArray<FPaintableTexture> NewCommonPaintableTextures;
+							for (FPaintableTexture const& PaintableTexture : PaintableTextures)
+							{
+								if (CommonPaintableTextures.Contains(PaintableTexture))
+								{
+									NewCommonPaintableTextures.Add(PaintableTexture);
+								}
+							}
+							CommonPaintableTextures = MoveTemp(NewCommonPaintableTextures);
+						}
+					}
 				}
 			}
 		}
 	}
 
-	// Texture Asset painting only supports single component select.
-	bSelectionSupportsTextureAssetPaint &= SelectedMeshComponents.Num() == 1;
+	// Texture Asset painting requires some texture that is common to all selected components.
+	CommonPaintableTextures.RemoveAll([](FPaintableTexture const& PaintableTexture) { return PaintableTexture.bIsMeshTexture; });
+	bSelectionSupportsTextureAssetPaint &= CommonPaintableTextures.Num() > 0;
 }
 
 TArray<FPerComponentVertexColorData> UMeshPaintingSubsystem::GetCopiedColorsByComponent() const
@@ -1705,14 +1731,24 @@ void UMeshPaintingSubsystem::CacheSelectionData(const int32 PaintLODIndex, const
 	bNeedsRecache = false;
 }
 
-int32 UMeshPaintingSubsystem::GetMaxUVIndexToPaint() const
+FIntPoint UMeshPaintingSubsystem::GetMinMaxUVChannelsToPaint() const
 {
-	if (PaintableComponents.Num() == 1 && PaintableComponents[0].IsValid())
+	// For multiple selection take the value from the component with fewest UV channels.
+	int32 MinNumUVs = INDEX_NONE;
+	int32 MaxNumUVs = INDEX_NONE;
+	for (TWeakObjectPtr<UMeshComponent> PaintableComponent : PaintableComponents)
 	{
-		return GEngine->GetEngineSubsystem<UMeshPaintingSubsystem>()->GetNumberOfUVs(PaintableComponents[0].Get(), 0) - 1;
+		if (UMeshComponent* MeshComponent = PaintableComponents[0].Get())
+		{
+			const int32 NumUVs = GetNumberOfUVs(MeshComponent, 0);
+			MinNumUVs = MinNumUVs == INDEX_NONE ? NumUVs : FMath::Min(MinNumUVs, NumUVs);
+			MaxNumUVs = MaxNumUVs == INDEX_NONE ? NumUVs : FMath::Max(MaxNumUVs, NumUVs);
+		}
 	}
 
-	return 0;
+	MinNumUVs = MinNumUVs == INDEX_NONE ? 0 : MinNumUVs;
+	MaxNumUVs = MaxNumUVs == INDEX_NONE ? 0 : MaxNumUVs;
+	return FIntPoint(MinNumUVs, MaxNumUVs);
 }
 
 TCHAR const* FMeshPaintToolSettingHelpers::GetCacheIdentifier()
diff --git a/Engine/Plugins/MeshPainting/Source/MeshPaintingToolset/Private/MeshTexturePaintingTool.cpp b/Engine/Plugins/MeshPainting/Source/MeshPaintingToolset/Private/MeshTexturePaintingTool.cpp
--- a/Engine/Plugins/MeshPainting/Source/MeshPaintingToolset/Private/MeshTexturePaintingTool.cpp
+++ b/Engine/Plugins/MeshPainting/Source/MeshPaintingToolset/Private/MeshTexturePaintingTool.cpp
@@ -234,13 +234,14 @@ void UMeshTexturePaintingTool::OnTick(float DeltaTime)
 				}
 
 				// Painting textures
-				if ((TexturePaintingCurrentMeshComponent != nullptr) && (TexturePaintingCurrentMeshComponent != SelectedComponent))
+				UTexture2D* SelectedPaintTexure = GetSelectedPaintTexture(SelectedComponent);
+				if (PaintingTexture2D != nullptr && PaintingTexture2D != SelectedPaintTexure)
 				{
-					// Mesh has changed, so finish up with our previous texture
+					// Texture has changed, so finish up with our previous texture
 					FinishPaintingTexture();
 				}
 
-				if (TexturePaintingCurrentMeshComponent == nullptr)
+				if (PaintingTexture2D == nullptr)
 				{
 					StartPaintingTexture(SelectedComponent, *MeshAdapter);
 				}
@@ -277,7 +278,7 @@ void UMeshTexturePaintingTool::OnTick(float DeltaTime)
 	// Wait till end of the tick to finish painting so all systems in-between know if we've painted this frame
 	if (bRequestPaintBucketFill)
 	{
-		if (TexturePaintingCurrentMeshComponent != nullptr)
+		if (PaintingTexture2D != nullptr)
 		{
 			FinishPaintingTexture();
 			FinishPainting();
@@ -368,7 +369,8 @@ bool UMeshTexturePaintingTool::PaintInternal(const TArrayView<TPair<FVector, FVe
 	const bool bIsPainting = (PaintAction == EMeshPaintModeAction::Paint);
 	const float InStrengthScale = PaintStrength;;
 
-	bool bPaintApplied = false;
+	bool bAnyPaintApplied = false;
+
 	UMeshPaintingSubsystem* MeshPaintingSubsystem = GEngine->GetEngineSubsystem<UMeshPaintingSubsystem>();
 	// Fire out a ray to see if there is a *selected* component under the mouse cursor that can be painted.
 	for (int32 i = 0; i < Rays.Num(); ++i)
@@ -508,6 +510,7 @@ bool UMeshTexturePaintingTool::PaintInternal(const TArrayView<TPair<FVector, FVe
 			TArray<FTexturePaintMeshSectionInfo> MaterialSections;
 			UTexturePaintToolset::RetrieveMeshSectionsForTextures(HoveredComponent, 0/*CachedLODIndex*/, Textures, MaterialSections);
 
+			bool bPaintApplied = false;
 			TArray<FTexturePaintTriangleInfo> TrianglePaintInfoArray;
 			for (int32 PaintRayResultId : PaintRayResultIds)
 			{
@@ -518,14 +521,22 @@ bool UMeshTexturePaintingTool::PaintInternal(const TArrayView<TPair<FVector, FVe
 				break;
 			}
 
+			if (!bPaintApplied)
+			{
+				continue;
+			}
+
 			// Painting textures
-			if ((TexturePaintingCurrentMeshComponent != nullptr) && (TexturePaintingCurrentMeshComponent != HoveredComponent))
+			bAnyPaintApplied = true;
+
+			UTexture2D* SelectedPaintTexure = GetSelectedPaintTexture(HoveredComponent);
+			if (PaintingTexture2D != nullptr && PaintingTexture2D != SelectedPaintTexure)
 			{
-				// Mesh has changed, so finish up with our previous texture
+				// Texture has changed, so finish up with our previous texture
 				FinishPaintingTexture();
 			}
 
-			if (TexturePaintingCurrentMeshComponent == nullptr)
+			if (PaintingTexture2D == nullptr)
 			{
 				StartPaintingTexture(HoveredComponent, *MeshAdapter);
 			}
@@ -549,7 +560,7 @@ bool UMeshTexturePaintingTool::PaintInternal(const TArrayView<TPair<FVector, FVe
 	}
 
 	LastPaintRayResults = PaintRayResults;
-	return bPaintApplied;
+	return bAnyPaintApplied;
 }
 
 /** Painting texture to use in material override should be the virtual texture adapter if it exists. */
@@ -1525,44 +1536,80 @@ void UMeshTextureAssetPaintingTool::CacheTexturePaintData()
 	UMeshPaintingSubsystem* MeshPaintingSubsystem = GEngine->GetEngineSubsystem<UMeshPaintingSubsystem>();
 	if (MeshPaintingSubsystem)
 	{
-		PaintableTextures.Empty();
+		PaintableTextures.Reset();
 
 		UTexture* DefaultTexture = nullptr;
 		int32 DefaultUVChannelIndex = INDEX_NONE;
 
 		TArray<UMeshComponent*> PaintableComponents = MeshPaintingSubsystem->GetPaintableMeshComponents();
-		if (PaintableComponents.IsValidIndex(0))
+
+		// Gather textures on first component.
+		if (PaintableComponents.Num() > 0)
 		{
-			int32 DefaultTextureIndex = INDEX_NONE;
 			TSharedPtr<IMeshPaintComponentAdapter> Adapter = MeshPaintingSubsystem->GetAdapterForComponent(PaintableComponents[0]);
+			int32 DefaultTextureIndex = INDEX_NONE;
 			UTexturePaintToolset::RetrieveTexturesForComponent(PaintableComponents[0], Adapter.Get(), DefaultTextureIndex, PaintableTextures);
-			if (PaintableTextures.IsValidIndex(DefaultTextureIndex))
+			if (DefaultTexture == nullptr && PaintableTextures.IsValidIndex(DefaultTextureIndex))
 			{
 				DefaultTexture = PaintableTextures[DefaultTextureIndex].Texture;
 				DefaultUVChannelIndex = PaintableTextures[DefaultTextureIndex].UVChannelIndex;
 			}
 		}
-		
+
+		// If there is more than one component we only want textures that are referenced by ALL selected components.
+		for (int32 Index = 1; Index < PaintableComponents.Num(); ++Index)
+		{
+			UMeshComponent* PaintableComponent = PaintableComponents[Index];
+			TSharedPtr<IMeshPaintComponentAdapter> Adapter = MeshPaintingSubsystem->GetAdapterForComponent(PaintableComponent);
+			TArray<FPaintableTexture> PerComponentPaintableTextures;
+			int32 DefaultTextureIndex = INDEX_NONE;
+			UTexturePaintToolset::RetrieveTexturesForComponent(PaintableComponent, Adapter.Get(), DefaultTextureIndex, PerComponentPaintableTextures);
+			if (DefaultTexture == nullptr && PerComponentPaintableTextures.IsValidIndex(DefaultTextureIndex))
+			{
+				DefaultTexture = PerComponentPaintableTextures[DefaultTextureIndex].Texture;
+				DefaultUVChannelIndex = PerComponentPaintableTextures[DefaultTextureIndex].UVChannelIndex;
+			}
+
+			TArray<FPaintableTexture> CommonPaintableTextures;
+			for (FPaintableTexture const& PaintableTexture : PerComponentPaintableTextures)
+			{
+				if (PaintableTextures.Contains(PaintableTexture))
+				{
+					CommonPaintableTextures.Add(PaintableTexture);
+				}
+			}
+			PaintableTextures = MoveTemp(CommonPaintableTextures);
+		}
+
 		PaintableTextures.RemoveAll([](FPaintableTexture const& PaintableTexture) { return PaintableTexture.bIsMeshTexture; });
 
 		// Ensure that the selection remains valid or is invalidated
-		if (!PaintableTextures.Contains(AssetProperties->PaintTexture))
+		int32 SelectedIndex = INDEX_NONE;
+
+		if (AssetProperties->PaintTexture != nullptr)
 		{
-			UTexture2D* NewTexture = nullptr;
-			int32 NewUVChannelIndex = INDEX_NONE;
-			if (PaintableTextures.Contains(DefaultTexture))
+			// First try to find fully matching entry, then by texture only (a texture may appear with multiple UV channels).
+			SelectedIndex = PaintableTextures.Find(FPaintableTexture(AssetProperties->PaintTexture, AssetProperties->UVChannel));
+			if (SelectedIndex == INDEX_NONE)
 			{
-				NewTexture = Cast<UTexture2D>(DefaultTexture);
-				NewUVChannelIndex = DefaultUVChannelIndex;
+				SelectedIndex = PaintableTextures.IndexOfByPredicate([&](const FPaintableTexture& Texture) { return Texture.Texture == AssetProperties->PaintTexture; });
+			}
 		}
-			else if (PaintableTextures.Num() > 0)
+		if (SelectedIndex == INDEX_NONE && DefaultTexture != nullptr)
+		{
+			SelectedIndex = PaintableTextures.Find(FPaintableTexture(DefaultTexture, DefaultUVChannelIndex));
+			if (SelectedIndex == INDEX_NONE)
 			{
-				NewTexture = Cast<UTexture2D>(PaintableTextures[0].Texture);
-				NewUVChannelIndex = PaintableTextures[0].UVChannelIndex;
+				SelectedIndex = PaintableTextures.IndexOfByPredicate([&](const FPaintableTexture& Texture) { return Texture.Texture == DefaultTexture; });
 			}
-			AssetProperties->PaintTexture = NewTexture;
-			AssetProperties->UVChannel = NewUVChannelIndex;
 		}
+		if (SelectedIndex == INDEX_NONE && PaintableTextures.Num() > 0)
+		{
+			SelectedIndex = 0;
+		}
+
+		AssetProperties->PaintTexture = SelectedIndex == INDEX_NONE ? nullptr : Cast<UTexture2D>(PaintableTextures[SelectedIndex].Texture);
+		AssetProperties->UVChannel = SelectedIndex == INDEX_NONE ? INDEX_NONE : PaintableTextures[SelectedIndex].UVChannelIndex;
 	}
 }
 
diff --git a/Engine/Plugins/MeshPainting/Source/MeshPaintingToolset/Public/MeshPaintHelpers.h b/Engine/Plugins/MeshPainting/Source/MeshPaintingToolset/Public/MeshPaintHelpers.h
--- a/Engine/Plugins/MeshPainting/Source/MeshPaintingToolset/Public/MeshPaintHelpers.h
+++ b/Engine/Plugins/MeshPainting/Source/MeshPaintingToolset/Public/MeshPaintHelpers.h
@@ -271,7 +271,7 @@ public:
 	TArray<FPerComponentVertexColorData> GetCopiedColorsByComponent() const;
 	void SetCopiedColorsByComponent(TArray<FPerComponentVertexColorData>& InCopiedColors);
 	void CacheSelectionData(const int32 PaintLODIndex, const int32 UVChannel);
-	int32 GetMaxUVIndexToPaint() const;
+	FIntPoint GetMinMaxUVChannelsToPaint() const;
 	void ResetState();
 	void Refresh();
 	bool SelectionContainsPerLODColors() const { return bSelectionContainsPerLODColors; }
diff --git a/Engine/Plugins/MeshPainting/Source/MeshPaintingToolset/Public/MeshPaintingToolsetTypes.h b/Engine/Plugins/MeshPainting/Source/MeshPaintingToolset/Public/MeshPaintingToolsetTypes.h
--- a/Engine/Plugins/MeshPainting/Source/MeshPaintingToolset/Public/MeshPaintingToolsetTypes.h
+++ b/Engine/Plugins/MeshPainting/Source/MeshPaintingToolset/Public/MeshPaintingToolsetTypes.h
@@ -123,8 +123,7 @@ struct FPaintableTexture
 	/** Overloaded equality operator for use with TArrays Contains method. */
 	bool operator==(const FPaintableTexture& rhs) const
 	{
-		return (Texture == rhs.Texture);
-		/* && (UVChannelIndex == rhs.UVChannelIndex);*/// if we compare UVChannel we would have to duplicate the texture
+		return Texture == rhs.Texture && UVChannelIndex == rhs.UVChannelIndex && bIsMeshTexture == rhs.bIsMeshTexture;
 	}
 };
 
diff --git a/Engine/Plugins/MeshPainting/Source/MeshPaintingToolset/Public/MeshTexturePaintingTool.h b/Engine/Plugins/MeshPainting/Source/MeshPaintingToolset/Public/MeshTexturePaintingTool.h
--- a/Engine/Plugins/MeshPainting/Source/MeshPaintingToolset/Public/MeshTexturePaintingTool.h
+++ b/Engine/Plugins/MeshPainting/Source/MeshPaintingToolset/Public/MeshTexturePaintingTool.h
@@ -299,7 +299,7 @@ protected:
 	// End UInteractiveTool Interface.
 
 	// Begin UMeshTexturePaintingTool Interface.
-	virtual bool AllowsMultiselect() const override { return false; }
+	virtual bool AllowsMultiselect() const override { return true; }
 	virtual bool IsMeshAdapterSupported(TSharedPtr<IMeshPaintComponentAdapter> MeshAdapter) const override;
 	virtual void OnPropertyModified(UObject* PropertySet, FProperty* Property) override;
 	virtual UTexture2D* GetSelectedPaintTexture(UMeshComponent const* InMeshComponent) const override;
