MeshPaint: Support standard UE hotkeys on brush size/strength/falloff.

diff --git a/Engine/Plugins/MeshPainting/Source/MeshPaintEditorMode/Private/MeshPaintMode.cpp b/Engine/Plugins/MeshPainting/Source/MeshPaintEditorMode/Private/MeshPaintMode.cpp
--- a/Engine/Plugins/MeshPainting/Source/MeshPaintEditorMode/Private/MeshPaintMode.cpp
+++ b/Engine/Plugins/MeshPainting/Source/MeshPaintEditorMode/Private/MeshPaintMode.cpp
@@ -392,6 +392,36 @@ void UMeshPaintMode::BindCommands()
 	CommandList->MapAction(Commands.NextTexture, FUIAction(
 		FExecuteAction::CreateUObject(this, &UMeshPaintMode::CycleTextures, 1),
 		FCanExecuteAction::CreateUObject(this, &UMeshPaintMode::CanCycleTextures)));
+
+	CommandList->MapAction(Commands.IncreaseBrushRadius, FUIAction(
+		FExecuteAction::CreateUObject(this, &UMeshPaintMode::ChangeBrushRadius, 1),
+		FCanExecuteAction::CreateUObject(this, &UMeshPaintMode::CanChangeBrush),
+		EUIActionRepeatMode::RepeatEnabled));
+	
+	CommandList->MapAction(Commands.DecreaseBrushRadius, FUIAction(
+		FExecuteAction::CreateUObject(this, &UMeshPaintMode::ChangeBrushRadius, -1),
+		FCanExecuteAction::CreateUObject(this, &UMeshPaintMode::CanChangeBrush),
+		EUIActionRepeatMode::RepeatEnabled));
+
+	CommandList->MapAction(Commands.IncreaseBrushStrength, FUIAction(
+		FExecuteAction::CreateUObject(this, &UMeshPaintMode::ChangeBrushStrength, 1),
+		FCanExecuteAction::CreateUObject(this, &UMeshPaintMode::CanChangeBrush),
+		EUIActionRepeatMode::RepeatEnabled));
+	
+	CommandList->MapAction(Commands.DecreaseBrushStrength, FUIAction(
+		FExecuteAction::CreateUObject(this, &UMeshPaintMode::ChangeBrushStrength, -1),
+		FCanExecuteAction::CreateUObject(this, &UMeshPaintMode::CanChangeBrush),
+		EUIActionRepeatMode::RepeatEnabled));
+
+	CommandList->MapAction(Commands.IncreaseBrushFalloff, FUIAction(
+		FExecuteAction::CreateUObject(this, &UMeshPaintMode::ChangeBrushFalloff, 1),
+		FCanExecuteAction::CreateUObject(this, &UMeshPaintMode::CanChangeBrush),
+		EUIActionRepeatMode::RepeatEnabled));
+
+	CommandList->MapAction(Commands.DecreaseBrushFalloff, FUIAction(
+		FExecuteAction::CreateUObject(this, &UMeshPaintMode::ChangeBrushFalloff, -1),
+		FCanExecuteAction::CreateUObject(this, &UMeshPaintMode::CanChangeBrush),
+		EUIActionRepeatMode::RepeatEnabled));
 }
 
 void UMeshPaintMode::OnToolStarted(UInteractiveToolManager* Manager, UInteractiveTool* Tool)
@@ -950,6 +980,60 @@ bool UMeshPaintMode::CanCycleTextures() const
 	return TexturePaintingTool != nullptr;
 }
 
+void UMeshPaintMode::ChangeBrushRadius(int32 Direction)
+{
+	if (UBaseBrushTool* Tool = Cast<UBaseBrushTool>(GetToolManager()->GetActiveTool(EToolSide::Left)))
+	{
+		FScopedTransaction Transaction(LOCTEXT("LevelMeshPainter_ChangeBrushRadius", "Changed Brush Radius"));
+		if (Direction > 0)
+		{
+			Tool->IncreaseBrushSizeAction();
+		}
+		else
+		{
+			Tool->DecreaseBrushSizeAction();
+		}
+	}
+}
+
+void UMeshPaintMode::ChangeBrushStrength(int32 Direction)
+{
+	if (UBaseBrushTool* Tool = Cast<UBaseBrushTool>(GetToolManager()->GetActiveTool(EToolSide::Left)))
+	{
+		FScopedTransaction Transaction(LOCTEXT("LevelMeshPainter_ChangeBrushStrength", "Changed Brush Strength"));
+		if (Direction > 0)
+		{
+			Tool->IncreaseBrushStrengthAction();
+		}
+		else
+		{
+			Tool->DecreaseBrushStrengthAction();
+		}
+	}
+}
+
+void UMeshPaintMode::ChangeBrushFalloff(int32 Direction)
+{
+	if (UBaseBrushTool* Tool = Cast<UBaseBrushTool>(GetToolManager()->GetActiveTool(EToolSide::Left)))
+	{
+		FScopedTransaction Transaction(LOCTEXT("LevelMeshPainter_ChangeBrushFalloff", "Changed Brush Falloff"));
+		if (Direction > 0)
+		{
+			Tool->IncreaseBrushFalloffAction();
+		}
+		else
+		{
+			Tool->DecreaseBrushFalloffAction();
+		}
+	}
+}
+
+bool UMeshPaintMode::CanChangeBrush() const
+{
+	UBaseBrushTool* Tool = Cast<UBaseBrushTool>(GetToolManager()->GetActiveTool(EToolSide::Left));
+	return Tool != nullptr;
+}
+
 void UMeshPaintMode::ActivateDefaultTool()
 {
 	FName PaletteName = Toolkit->GetCurrentPalette();
diff --git a/Engine/Plugins/MeshPainting/Source/MeshPaintEditorMode/Private/MeshPaintMode.h b/Engine/Plugins/MeshPainting/Source/MeshPaintEditorMode/Private/MeshPaintMode.h
--- a/Engine/Plugins/MeshPainting/Source/MeshPaintEditorMode/Private/MeshPaintMode.h
+++ b/Engine/Plugins/MeshPainting/Source/MeshPaintEditorMode/Private/MeshPaintMode.h
@@ -142,6 +142,10 @@ protected:
 	bool CanCycleMeshLODs() const;
 	void CycleTextures(int32 Direction);
 	bool CanCycleTextures() const;
+	void ChangeBrushRadius(int32 Direction);
+	void ChangeBrushStrength(int32 Direction);
+	void ChangeBrushFalloff(int32 Direction);
+	bool CanChangeBrush() const;
 	// End command bindings
 
 protected:
diff --git a/Engine/Plugins/MeshPainting/Source/MeshPaintEditorMode/Private/MeshPaintModeCommands.cpp b/Engine/Plugins/MeshPainting/Source/MeshPaintEditorMode/Private/MeshPaintModeCommands.cpp
--- a/Engine/Plugins/MeshPainting/Source/MeshPaintEditorMode/Private/MeshPaintModeCommands.cpp
+++ b/Engine/Plugins/MeshPainting/Source/MeshPaintEditorMode/Private/MeshPaintModeCommands.cpp
@@ -62,6 +62,12 @@ void FMeshPaintEditorModeCommands::RegisterCommands()
 	UI_COMMAND(NextLOD, "Next LOD", "Cycles to the next possible mesh LOD to paint on", EUserInterfaceActionType::Button, FInputChord(EKeys::N));
 	UI_COMMAND(PreviousTexture, "Previous Texture", "Cycle To previous texture", EUserInterfaceActionType::Button, FInputChord(EKeys::Comma));
 	UI_COMMAND(NextTexture, "Next Texture", "Cycle To next texture", EUserInterfaceActionType::Button, FInputChord(EKeys::Period));
+	UI_COMMAND(IncreaseBrushRadius, "Increase Brush Radius", "Increase brush radius by a percentage of its current size.", EUserInterfaceActionType::Button, FInputChord(EKeys::RightBracket));
+	UI_COMMAND(DecreaseBrushRadius, "Decrease Brush Size", "Decrease brush radius by a percentage of its current size.", EUserInterfaceActionType::Button, FInputChord(EKeys::LeftBracket));
+	UI_COMMAND(IncreaseBrushStrength, "Increase Brush Strength", "Increase brush strength by a fixed increment.", EUserInterfaceActionType::Button, FInputChord(EModifierKey::Control, EKeys::RightBracket));
+	UI_COMMAND(DecreaseBrushStrength, "Decrease Brush Strength", "Decrease brush strength by a fixed increment.", EUserInterfaceActionType::Button, FInputChord(EModifierKey::Control, EKeys::LeftBracket));
+	UI_COMMAND(IncreaseBrushFalloff, "Increase Brush Falloff", "Increase brush falloff by a fixed increment.", EUserInterfaceActionType::Button, FInputChord(FInputChord(EModifierKey::Control | EModifierKey::Shift, EKeys::RightBracket)));
+	UI_COMMAND(DecreaseBrushFalloff, "Decrease Brush Falloff", "Decrease brush falloff by a fixed increment.", EUserInterfaceActionType::Button, FInputChord(FInputChord(EModifierKey::Control | EModifierKey::Shift, EKeys::LeftBracket)));
 
 	TArray<TSharedPtr<FUICommandInfo>> VertexColorCommands = {
 		SelectVertex, PaintVertexColor, SwapColor, FillVertex, PropagateLODs, PropagateMesh, SaveVertex, RemoveVertex, Copy, Paste, Import, GetTextureColors, FixVertex
diff --git a/Engine/Plugins/MeshPainting/Source/MeshPaintEditorMode/Public/MeshPaintModeCommands.h b/Engine/Plugins/MeshPainting/Source/MeshPaintEditorMode/Public/MeshPaintModeCommands.h
--- a/Engine/Plugins/MeshPainting/Source/MeshPaintEditorMode/Public/MeshPaintModeCommands.h
+++ b/Engine/Plugins/MeshPainting/Source/MeshPaintEditorMode/Public/MeshPaintModeCommands.h
@@ -62,6 +62,13 @@ public:
 	TSharedPtr<FUICommandInfo> PreviousTexture;
 	TSharedPtr<FUICommandInfo> NextTexture;
 
+	TSharedPtr<FUICommandInfo> IncreaseBrushRadius;
+	TSharedPtr<FUICommandInfo> DecreaseBrushRadius;
+	TSharedPtr<FUICommandInfo> IncreaseBrushStrength;
+	TSharedPtr<FUICommandInfo> DecreaseBrushStrength;
+	TSharedPtr<FUICommandInfo> IncreaseBrushFalloff;
+	TSharedPtr<FUICommandInfo> DecreaseBrushFalloff;
+
 	TMap<FName, TArray<TSharedPtr<FUICommandInfo>>> Commands;
 };
 
