MeshPaint: Fix bug where max UV channel was incorrectly calculated.

diff --git a/Engine/Plugins/MeshPainting/Source/MeshPaintingToolset/Private/MeshPaintHelpers.cpp b/Engine/Plugins/MeshPainting/Source/MeshPaintingToolset/Private/MeshPaintHelpers.cpp
--- a/Engine/Plugins/MeshPainting/Source/MeshPaintingToolset/Private/MeshPaintHelpers.cpp
+++ b/Engine/Plugins/MeshPainting/Source/MeshPaintingToolset/Private/MeshPaintHelpers.cpp
@@ -1738,7 +1738,7 @@ FIntPoint UMeshPaintingSubsystem::GetMinMaxUVChannelsToPaint() const
 	int32 MaxNumUVs = INDEX_NONE;
 	for (TWeakObjectPtr<UMeshComponent> PaintableComponent : PaintableComponents)
 	{
-		if (UMeshComponent* MeshComponent = PaintableComponents[0].Get())
+		if (UMeshComponent* MeshComponent = PaintableComponent.Get())
 		{
 			const int32 NumUVs = GetNumberOfUVs(MeshComponent, 0);
 			MinNumUVs = MinNumUVs == INDEX_NONE ? NumUVs : FMath::Min(MinNumUVs, NumUVs);
