Mesh Paint: Fix seam painting to a shared texture asset when using multi-select.
Stop tracking a single currently painted component but track all components that we paint a texture for.
Rebuild seam mask each time we change components.
Also removed some unused structures.

diff --git a/Engine/Plugins/MeshPainting/Source/MeshPaintingToolset/Private/MeshTexturePaintingTool.cpp b/Engine/Plugins/MeshPainting/Source/MeshPaintingToolset/Private/MeshTexturePaintingTool.cpp
--- a/Engine/Plugins/MeshPainting/Source/MeshPaintingToolset/Private/MeshTexturePaintingTool.cpp
+++ b/Engine/Plugins/MeshPainting/Source/MeshPaintingToolset/Private/MeshTexturePaintingTool.cpp
@@ -247,7 +247,7 @@ void UMeshTexturePaintingTool::OnTick(float DeltaTime)
 				}
 
 				FMeshPaintParameters* LastParams = nullptr;
-				PaintTexture(BucketFillParams, UVChannel, TrianglePaintInfoArray, *MeshAdapter, LastParams);
+				PaintTexture(BucketFillParams, UVChannel, TrianglePaintInfoArray, SelectedComponent, *MeshAdapter, LastParams);
 			}
 		}
 
@@ -541,7 +541,7 @@ bool UMeshTexturePaintingTool::PaintInternal(const TArrayView<TPair<FVector, FVe
 				StartPaintingTexture(HoveredComponent, *MeshAdapter);
 			}
 
-			if (TexturePaintingCurrentMeshComponent != nullptr)
+			if (PaintingTexture2D != nullptr)
 			{
 				for (int32 PaintRayResultId : PaintRayResultIds)
 				{
@@ -552,7 +552,7 @@ bool UMeshTexturePaintingTool::PaintInternal(const TArrayView<TPair<FVector, FVe
 						LastParams = &LastPaintRayResults[PaintRayResultId].Params;
 					}
 
-					PaintTexture(Params, UVChannel, TrianglePaintInfoArray, *MeshAdapter, LastParams);
+					PaintTexture(Params, UVChannel, TrianglePaintInfoArray, HoveredComponent, *MeshAdapter, LastParams);
 					break;
 				}
 			}
@@ -733,7 +733,6 @@ void UMeshTexturePaintingTool::GatherTextureTriangles(IMeshPaintComponentAdapter
 void UMeshTexturePaintingTool::StartPaintingTexture(UMeshComponent* InMeshComponent, const IMeshPaintComponentAdapter& GeometryInfo)
 {
 	check(InMeshComponent != nullptr);
-	check(TexturePaintingCurrentMeshComponent == nullptr);
 	check(PaintingTexture2D == nullptr);
 
 	// Only start new transaction if not in one currently
@@ -813,8 +812,6 @@ void UMeshTexturePaintingTool::StartPaintingTexture(UMeshComponent* InMeshCompon
 					TextureData->BrushMaskRenderTargetTexture->UpdateResourceImmediate();
 					TextureData->BrushMaskRenderTargetTexture->AddressX = TextureData->PaintRenderTargetTexture->AddressX;
 					TextureData->BrushMaskRenderTargetTexture->AddressY = TextureData->PaintRenderTargetTexture->AddressY;
-
-					TextureData->bGenerateSeamMask = true;
 				}
 
 				// Create the rendertarget used to store a texture seam mask
@@ -831,6 +828,7 @@ void UMeshTexturePaintingTool::StartPaintingTexture(UMeshComponent* InMeshCompon
 					TextureData->SeamMaskRenderTargetTexture->UpdateResourceImmediate();
 					TextureData->SeamMaskRenderTargetTexture->AddressX = TextureData->PaintRenderTargetTexture->AddressX;
 					TextureData->SeamMaskRenderTargetTexture->AddressY = TextureData->PaintRenderTargetTexture->AddressY;
+					TextureData->SeamMaskComponent = nullptr;
 				}
 			}
 
@@ -860,8 +858,6 @@ void UMeshTexturePaintingTool::StartPaintingTexture(UMeshComponent* InMeshCompon
 				TextureData->PaintBrushRenderTargetTexture = nullptr;
 			}
 
-			TexturePaintingCurrentMeshComponent = InMeshComponent;
-
 			check(Texture2D != nullptr);
 			PaintingTexture2D = Texture2D;
 		}
@@ -872,8 +868,6 @@ void UMeshTexturePaintingTool::StartPaintingTexture(UMeshComponent* InMeshCompon
 
 	if (bIsSourceTextureStreamedIn && bStartedPainting)
 	{
-		TexturePaintingCurrentMeshComponent = InMeshComponent;
-
 		check(Texture2D != nullptr);
 		PaintingTexture2D = Texture2D;
 
@@ -884,7 +878,7 @@ void UMeshTexturePaintingTool::StartPaintingTexture(UMeshComponent* InMeshCompon
 	}
 }
 
-void UMeshTexturePaintingTool::PaintTexture(FMeshPaintParameters& InParams, int32 UVChannel, TArray<FTexturePaintTriangleInfo>& InInfluencedTriangles, const IMeshPaintComponentAdapter& GeometryInfo, FMeshPaintParameters* LastParams)
+void UMeshTexturePaintingTool::PaintTexture(FMeshPaintParameters& InParams, int32 UVChannel, TArray<FTexturePaintTriangleInfo>& InInfluencedTriangles, UMeshComponent* MeshComponent, const IMeshPaintComponentAdapter& GeometryInfo, FMeshPaintParameters* LastParams)
 {
 	// We bail early if there are no influenced triangles
 	if (InInfluencedTriangles.Num() <= 0)
@@ -1109,6 +1103,7 @@ void UMeshTexturePaintingTool::PaintTexture(FMeshPaintParameters& InParams, int3
 		BrushPaintCanvas.Flush_GameThread(true);
 
 		TextureData->bIsPaintingTexture2DModified = true;
+		TextureData->PaintedComponents.AddUnique(MeshComponent);
 	}
 
 	ENQUEUE_RENDER_COMMAND(UpdateMeshPaintRTCommand1)(
@@ -1145,11 +1140,13 @@ void UMeshTexturePaintingTool::PaintTexture(FMeshPaintParameters& InParams, int3
 		const float MaxX = TextureData->PaintRenderTargetTexture->GetSurfaceWidth();
 		const float MaxY = TextureData->PaintRenderTargetTexture->GetSurfaceHeight();
 
-		if (TextureData->bGenerateSeamMask == true)
+		if (TextureData->SeamMaskComponent != MeshComponent)
 		{
-			// Generate the texture seam mask.  This is a slow operation when the object has many triangles so we only do it once when painting is started.
-			UTexturePaintToolset::GenerateSeamMask(TexturePaintingCurrentMeshComponent, UVChannel, TextureData->SeamMaskRenderTargetTexture, TextureData->PaintingTexture2D, TextureData->PaintRenderTargetTexture);
-			TextureData->bGenerateSeamMask = false;
+			// Generate the texture seam mask.  This is a slow operation when the object has many triangles so we try to only do it once when painting is started.
+			// @todo MeshPaint: We only store one seam mask, so when we are painting to a texture asset with multi-select we will end up re-rendering the mask each time the brush crosses component boundaries.
+			// Better could be to store a mask per component instead, but still lazily generate them on demand.
+			UTexturePaintToolset::GenerateSeamMask(MeshComponent, UVChannel, TextureData->SeamMaskRenderTargetTexture, TextureData->PaintingTexture2D, TextureData->PaintRenderTargetTexture);
+			TextureData->SeamMaskComponent = MeshComponent;
 		}
 
 		FTextureRenderTargetResource* RenderTargetResource = TextureData->PaintRenderTargetTexture->GameThread_GetRenderTargetResource();
@@ -1261,12 +1258,15 @@ void UMeshTexturePaintingTool::FinishPaintingTexture()
 
 			TextureData->bIsPaintingTexture2DModified = false;
 
-			OnPaintingFinishedDelegate.ExecuteIfBound(TexturePaintingCurrentMeshComponent);
+			for (UMeshComponent* PaintedComponent : TextureData->PaintedComponents)
+			{
+				OnPaintingFinishedDelegate.ExecuteIfBound(PaintedComponent);
+			}
+			TextureData->PaintedComponents.Reset();
 		}
 	}
 
 	PaintingTexture2D = nullptr;
-	TexturePaintingCurrentMeshComponent = nullptr;
 }
 
 void UMeshTexturePaintingTool::ClearAllTextureOverrides()
diff --git a/Engine/Plugins/MeshPainting/Source/MeshPaintingToolset/Public/MeshPaintingToolsetTypes.h b/Engine/Plugins/MeshPainting/Source/MeshPaintingToolset/Public/MeshPaintingToolsetTypes.h
--- a/Engine/Plugins/MeshPainting/Source/MeshPaintingToolset/Public/MeshPaintingToolsetTypes.h
+++ b/Engine/Plugins/MeshPainting/Source/MeshPaintingToolset/Public/MeshPaintingToolsetTypes.h
@@ -8,7 +8,6 @@
 
 #include "MeshPaintingToolsetTypes.generated.h"
 
-class UMaterialInterface;
 class UMeshComponent;
 class UTexture2D;
 class UTexture;
@@ -50,17 +49,6 @@ enum class EMeshPaintModeAction : uint8
 	Erase
 };
 
-/** Vertex paint target */
-UENUM()
-enum class EMeshVertexPaintModeTarget : uint8
-{
-	/** Paint the static mesh component instance in the level */
-	ComponentInstance,
-
-	/** Paint the actual static mesh asset */
-	Mesh
-};
-
 /** Mesh paint parameters */
 class FMeshPaintParameters
 {
@@ -130,7 +118,14 @@ struct FPaintableTexture
 USTRUCT()
 struct FPaintTexture2DData
 {
-	GENERATED_BODY()
+	GENERATED_USTRUCT_BODY()
+
+	FPaintTexture2DData() = default;
+
+	FPaintTexture2DData(UTexture2D* InPaintingTexture2D)
+		: PaintingTexture2D(InPaintingTexture2D)
+	{
+	}
 
 	/** The original texture that we're painting */
 	UPROPERTY(Transient)
@@ -144,10 +139,6 @@ struct FPaintTexture2DData
 	UPROPERTY(Transient)
 	TObjectPtr<class UVirtualTextureAdapter> PaintRenderTargetTextureAdapter;
 
-	/** Array of components that have the PaintRenderTargetTexture set as a texture override */
-	UPROPERTY(Transient)
-	TArray<TObjectPtr<UMeshComponent>> TextureOverrideComponents;
-
 	/** Optional render target texture used as an input while painting that contains a clone of the texture painting brush */
 	UPROPERTY(Transient)
 	TObjectPtr<UTextureRenderTarget2D> PaintBrushRenderTargetTexture;
@@ -160,34 +151,23 @@ struct FPaintTexture2DData
 	UPROPERTY(Transient)
 	TObjectPtr<UTextureRenderTarget2D> BrushMaskRenderTargetTexture;
 
-	/** True if we need to generate a texture seam mask used for texture dilation */
-	UPROPERTY(Transient)
-	bool bGenerateSeamMask = false;
-
 	/** Optional render target used to store generated mask for texture seams. We create this by projecting object triangles into texture space using the selected UV channel. */
 	UPROPERTY(Transient)
 	TObjectPtr<UTextureRenderTarget2D> SeamMaskRenderTargetTexture;
 
-	/** True if we have some painting applied to the PaintRenderTargetTexture. */
+	/** Array of components that have the PaintRenderTargetTexture set as a texture override */
 	UPROPERTY(Transient)
-	bool bIsPaintingTexture2DModified = false;
-
-	FPaintTexture2DData() = default;
-
-	FPaintTexture2DData(UTexture2D* InPaintingTexture2D)
-		: PaintingTexture2D(InPaintingTexture2D)
-	{
-	}
-};
-
-USTRUCT()
-struct FPaintComponentOverride
-{
-	GENERATED_BODY()
+	TArray<TObjectPtr<UMeshComponent>> TextureOverrideComponents;
 
-	/** List of components overridden */
+	/** Array of components that have been painted */
 	UPROPERTY(Transient)
 	TArray<TObjectPtr<UMeshComponent>> PaintedComponents;
+
+	/** Current component that was used to generate the texture seam mask used for texture dilation */
+	TObjectPtr<UMeshComponent> SeamMaskComponent;
+
+	/** True if we have some painting applied to the PaintRenderTargetTexture. */
+	bool bIsPaintingTexture2DModified = false;
 };
 
 namespace MeshPaintDefs
@@ -198,34 +178,3 @@ namespace MeshPaintDefs
 	static const int32 MaxSupportedPhysicalWeights = 4;
 	static const int32 MaxSupportedWeights = MaxSupportedPhysicalWeights + 1;
 }
-
-/**
-*  Wrapper to expose texture targets to WPF code.
-*/
-struct FTextureTargetListInfo
-{
-	UTexture2D* TextureData;
-	bool bIsSelected;
-	uint32 UndoCount;
-	uint32 UVChannelIndex;
-	FTextureTargetListInfo(UTexture2D* InTextureData, int32 InUVChannelIndex, bool InbIsSelected = false)
-		: TextureData(InTextureData)
-		, bIsSelected(InbIsSelected)
-		, UndoCount(0)
-		, UVChannelIndex(InUVChannelIndex)
-	{}
-};
-
-/**
-*  Wrapper to store which of a meshes materials is selected as well as the total number of materials.
-*/
-struct FMeshSelectedMaterialInfo
-{
-	int32 NumMaterials;
-	int32 SelectedMaterialIndex;
-
-	FMeshSelectedMaterialInfo(int32 InNumMaterials)
-		: NumMaterials(InNumMaterials)
-		, SelectedMaterialIndex(0)
-	{}
-};
diff --git a/Engine/Plugins/MeshPainting/Source/MeshPaintingToolset/Public/MeshTexturePaintingTool.h b/Engine/Plugins/MeshPainting/Source/MeshPaintingToolset/Public/MeshTexturePaintingTool.h
--- a/Engine/Plugins/MeshPainting/Source/MeshPaintingToolset/Public/MeshTexturePaintingTool.h
+++ b/Engine/Plugins/MeshPainting/Source/MeshPaintingToolset/Public/MeshTexturePaintingTool.h
@@ -183,7 +183,7 @@ private:
 	void GatherTextureTriangles(IMeshPaintComponentAdapter* Adapter, int32 TriangleIndex, const int32 VertexIndices[3], TArray<FTexturePaintTriangleInfo>* TriangleInfo, TArray<FTexturePaintMeshSectionInfo>* SectionInfos, int32 UVChannelIndex);
 	bool Paint(const FVector& InRayOrigin, const FVector& InRayDirection);
 	bool Paint(const TArrayView<TPair<FVector, FVector>>& Rays);
-	void PaintTexture(FMeshPaintParameters& InParams, int32 UVChannel, TArray<FTexturePaintTriangleInfo>& InInfluencedTriangles, const IMeshPaintComponentAdapter& GeometryInfo, FMeshPaintParameters* LastParams = nullptr);
+	void PaintTexture(FMeshPaintParameters& InParams, int32 UVChannel, TArray<FTexturePaintTriangleInfo>& InInfluencedTriangles, UMeshComponent* MeshComponent, const IMeshPaintComponentAdapter& GeometryInfo, FMeshPaintParameters* LastParams = nullptr);
 	bool PaintInternal(const TArrayView<TPair<FVector, FVector>>& Rays, EMeshPaintModeAction PaintAction, float PaintStrength);
 	void FinishPaintingTexture();
 	void FinishPainting();
@@ -207,10 +207,6 @@ private:
 	UPROPERTY(Transient)
 	TObjectPtr<UTexture2D> PaintingTexture2D;
 
-	/** The mesh component that we're currently painting on */
-	UPROPERTY(Transient)
-	TObjectPtr<UMeshComponent> TexturePaintingCurrentMeshComponent;
-
 	/** Hold the transaction while we are painting */
 	TUniquePtr<FScopedTransaction> PaintingTransaction;
 
