Fix mesh paint mode viewport settubg persistance issues.
Previously when realtime mode was off before going into mesh paint mode, it made an override to realtime when going into mesh paint mode but then didn't remove the override on exiting mesh paint mode. This is now fixed and we return to the original state when leaving mesh paint mode.
Previously when selecting a color view mode when in mesh paint mode it would break the reflection viewport mode. Now that is fixed and reflection view can be used during mesh paint mode without breaking.

diff --git a/Engine/Plugins/MeshPainting/Source/MeshPaintEditorMode/Private/MeshPaintMode.cpp b/Engine/Plugins/MeshPainting/Source/MeshPaintEditorMode/Private/MeshPaintMode.cpp
--- a/Engine/Plugins/MeshPainting/Source/MeshPaintEditorMode/Private/MeshPaintMode.cpp
+++ b/Engine/Plugins/MeshPainting/Source/MeshPaintEditorMode/Private/MeshPaintMode.cpp
@@ -287,7 +287,7 @@ void UMeshPaintMode::Tick(FEditorViewportClient* ViewportClient, float DeltaTime
 	if (ViewportClient->IsPerspective())
 	{
 		// Make sure perspective viewports are still set to real-time
-		GEditor->GetEditorSubsystem<UMeshPaintModeSubsystem>()->SetRealtimeViewport(true);
+		GEditor->GetEditorSubsystem<UMeshPaintModeSubsystem>()->SetRealtimeViewport(ViewportClient, true);
 
 		// Set viewport show flags
 		GEditor->GetEditorSubsystem<UMeshPaintModeSubsystem>()->SetViewportColorMode(CurrentActiveMode, ModeSettings->ColorViewMode, ViewportClient, ActiveTool);
@@ -1104,7 +1104,7 @@ void UMeshPaintMode::FillTexture()
 
 void UMeshPaintMode::OnResetViewMode()
 {
-	// Reset viewport color mode for all active viewports
+	// Reset viewport color mode and realtime override for all active viewports
 	for (FEditorViewportClient* ViewportClient : GEditor->GetAllViewportClients())
 	{
 		if (!ViewportClient || ViewportClient->GetModeTools() != GetModeManager())
@@ -1113,6 +1113,7 @@ void UMeshPaintMode::OnResetViewMode()
 		}
 
 		GEditor->GetEditorSubsystem<UMeshPaintModeSubsystem>()->SetViewportColorMode(EMeshPaintActiveMode::VertexColor, EMeshPaintDataColorViewMode::Normal, ViewportClient, nullptr);
+		GEditor->GetEditorSubsystem<UMeshPaintModeSubsystem>()->SetRealtimeViewport(ViewportClient, false);
 	}
 }
 
diff --git a/Engine/Plugins/MeshPainting/Source/MeshPaintEditorMode/Private/MeshPaintModeHelpers.cpp b/Engine/Plugins/MeshPainting/Source/MeshPaintEditorMode/Private/MeshPaintModeHelpers.cpp
--- a/Engine/Plugins/MeshPainting/Source/MeshPaintEditorMode/Private/MeshPaintModeHelpers.cpp
+++ b/Engine/Plugins/MeshPainting/Source/MeshPaintEditorMode/Private/MeshPaintModeHelpers.cpp
@@ -58,16 +58,10 @@ void UMeshPaintModeSubsystem::SetViewportColorMode(EMeshPaintActiveMode ActiveMo
 		// Update viewport show flags
 		{
 			// show flags forced on during vertex color modes
-			if (ColorViewMode == EMeshPaintDataColorViewMode::Normal)
-			{
-				ColorViewMode = EMeshPaintDataColorViewMode::Normal;
-			}
-
 			if (ColorViewMode == EMeshPaintDataColorViewMode::Normal)
 			{
 				if (ViewportClient->EngineShowFlags.VertexColors)
 				{
-					// If we're transitioning to normal mode then restore the backup
 					// Clear the flags relevant to vertex color modes
 					ViewportClient->EngineShowFlags.SetVertexColors(false);
 					ViewportClient->CurrentNaniteVisualizationMode = NAME_None;
@@ -83,12 +77,7 @@ void UMeshPaintModeSubsystem::SetViewportColorMode(EMeshPaintActiveMode ActiveMo
 			}
 			else
 			{
-				ViewportClient->EngineShowFlags.SetMaterials(true);
-				ViewportClient->EngineShowFlags.SetLighting(false);
-				ViewportClient->EngineShowFlags.SetBSPTriangles(true);
 				ViewportClient->EngineShowFlags.SetVertexColors(true);
-				ViewportClient->EngineShowFlags.SetPostProcessing(false);
-				ViewportClient->EngineShowFlags.SetHMDDistortion(false);
 
 				MeshPaintVisualize::SetShowMode(EMeshPaintVisualizeShowMode::ShowSelected);
 
@@ -154,23 +143,26 @@ void UMeshPaintModeSubsystem::SetViewportColorMode(EMeshPaintActiveMode ActiveMo
 	}
 }
 
-void UMeshPaintModeSubsystem::SetRealtimeViewport(bool bRealtime)
+void UMeshPaintModeSubsystem::SetRealtimeViewport(FEditorViewportClient* ViewportClient, bool bRealtime)
 {
-	FLevelEditorModule& LevelEditorModule = FModuleManager::GetModuleChecked<FLevelEditorModule>("LevelEditor");
-	TSharedPtr< IAssetViewport > ViewportWindow = LevelEditorModule.GetFirstActiveViewport();
-	if (ViewportWindow.IsValid())
+	if (ViewportClient != nullptr)
 	{
-		FEditorViewportClient& Viewport = ViewportWindow->GetAssetViewportClient();
-		if (Viewport.IsPerspective())
+		if (ViewportClient->IsPerspective())
 		{
-			const FText SystemDisplayName = NSLOCTEXT("MeshPaint", "RealtimeOverrideMessage_MeshPaint", "Mesh Paint");
+			static const FText SystemDisplayName = NSLOCTEXT("MeshPaint", "RealtimeOverrideMessage_MeshPaint", "Mesh Paint");
 			if (bRealtime)
 			{
-				Viewport.AddRealtimeOverride(bRealtime, SystemDisplayName);
+				if (!ViewportClient->HasRealtimeOverride(SystemDisplayName))
+				{
+					ViewportClient->AddRealtimeOverride(bRealtime, SystemDisplayName);
+				}
 			}
 			else
 			{
-				Viewport.RemoveRealtimeOverride(SystemDisplayName);
+				if (ViewportClient->HasRealtimeOverride(SystemDisplayName))
+				{
+					ViewportClient->RemoveRealtimeOverride(SystemDisplayName);
+				}
 			}
 		}
 	}
diff --git a/Engine/Plugins/MeshPainting/Source/MeshPaintEditorMode/Public/MeshPaintModeHelpers.h b/Engine/Plugins/MeshPainting/Source/MeshPaintEditorMode/Public/MeshPaintModeHelpers.h
--- a/Engine/Plugins/MeshPainting/Source/MeshPaintEditorMode/Public/MeshPaintModeHelpers.h
+++ b/Engine/Plugins/MeshPainting/Source/MeshPaintEditorMode/Public/MeshPaintModeHelpers.h
@@ -42,8 +42,8 @@ public:
 	/** Forces the Viewport Client to render using the given Viewport Color ViewMode */
 	void SetViewportColorMode(EMeshPaintActiveMode ActiveMode, EMeshPaintDataColorViewMode ColorViewMode, FEditorViewportClient* ViewportClient, UInteractiveTool const* ActiveTool);
 
-	/** Sets whether or not the level viewport should be real time rendered move or viewport as parameter? */
-	void SetRealtimeViewport(bool bRealtime);
+	/** Sets whether or not the viewport should be real time rendered */
+	void SetRealtimeViewport(FEditorViewportClient* ViewportClient, bool bRealtime);
 
 
 	/** Helper function to import Vertex Colors from a Texture to the specified MeshComponent (makes use of SImportVertexColorsOptions Widget) */
