Fix warning C6011: Dereferencing NULL pointer 'MeshPaintingSubsystem'

#rnx

diff --git a/Engine/Plugins/MeshPainting/Source/MeshPaintingToolset/Private/MeshTexturePaintingTool.cpp b/Engine/Plugins/MeshPainting/Source/MeshPaintingToolset/Private/MeshTexturePaintingTool.cpp
--- a/Engine/Plugins/MeshPainting/Source/MeshPaintingToolset/Private/MeshTexturePaintingTool.cpp
+++ b/Engine/Plugins/MeshPainting/Source/MeshPaintingToolset/Private/MeshTexturePaintingTool.cpp
@@ -363,27 +363,26 @@ bool UMeshTexturePaintingTool::PaintInternal(const TArrayView<TPair<FVector, FVe
 	TArray<FPaintRayResults> PaintRayResults;
 	PaintRayResults.AddDefaulted(Rays.Num());
 
-	TMap<UMeshComponent*, TArray<int32>> HoveredComponents;
-
-	const float BrushRadius = BrushProperties->BrushRadius;
-	const bool bIsPainting = (PaintAction == EMeshPaintModeAction::Paint);
-	const float InStrengthScale = PaintStrength;;
-
 	bool bAnyPaintApplied = false;
 
-	UMeshPaintingSubsystem* MeshPaintingSubsystem = GEngine->GetEngineSubsystem<UMeshPaintingSubsystem>();
-	// Fire out a ray to see if there is a *selected* component under the mouse cursor that can be painted.
-	for (int32 i = 0; i < Rays.Num(); ++i)
+	if (UMeshPaintingSubsystem* MeshPaintingSubsystem = GEngine->GetEngineSubsystem<UMeshPaintingSubsystem>())
 	{
-		const FVector& RayOrigin = Rays[i].Key;
-		const FVector& RayDirection = Rays[i].Value;
-		FHitResult& BestTraceResult = PaintRayResults[i].BestTraceResult;
+		TMap<UMeshComponent*, TArray<int32>> HoveredComponents;
 
-		const FVector TraceStart(RayOrigin);
-		const FVector TraceEnd(RayOrigin + RayDirection * HALF_WORLD_MAX);
+		const float BrushRadius = BrushProperties->BrushRadius;
+		const bool bIsPainting = (PaintAction == EMeshPaintModeAction::Paint);
+		const float InStrengthScale = PaintStrength;;
 
-		if (MeshPaintingSubsystem)
+		// Fire out a ray to see if there is a *selected* component under the mouse cursor that can be painted.
+		for (int32 i = 0; i < Rays.Num(); ++i)
 		{
+			const FVector& RayOrigin = Rays[i].Key;
+			const FVector& RayDirection = Rays[i].Value;
+			FHitResult& BestTraceResult = PaintRayResults[i].BestTraceResult;
+
+			const FVector TraceStart(RayOrigin);
+			const FVector TraceEnd(RayOrigin + RayDirection * HALF_WORLD_MAX);
+
 			for (UMeshComponent* MeshComponent : MeshPaintingSubsystem->GetPaintableMeshComponents())
 			{
 				TSharedPtr<IMeshPaintComponentAdapter> MeshAdapter = MeshPaintingSubsystem->GetAdapterForComponent(MeshComponent);
@@ -400,166 +399,168 @@ bool UMeshTexturePaintingTool::PaintInternal(const TArrayView<TPair<FVector, FVe
 					}
 				}
 			}
-		}
 
-		UMeshComponent* BestTraceMeshComponent = Cast<UMeshComponent>(BestTraceResult.GetComponent());
-		// If painting texture assets, just use the BestTraceMeshComponent as we only support painting a single mesh at a time in that mode.
-		const bool bAllowMultiselect = AllowsMultiselect();
+			UMeshComponent* BestTraceMeshComponent = Cast<UMeshComponent>(BestTraceResult.GetComponent());
+			// If painting texture assets, just use the BestTraceMeshComponent as we only support painting a single mesh at a time in that mode.
+			const bool bAllowMultiselect = AllowsMultiselect();
 
-		bool bUsed = false;
-		for (UMeshComponent* MeshComponent : MeshPaintingSubsystem->GetPaintableMeshComponents())
-		{
-			if (MeshComponent == BestTraceMeshComponent)
-			{
-				HoveredComponents.FindOrAdd(MeshComponent).Add(i);
-				bUsed = true;
-			}
-			else if (bAllowMultiselect)
+			bool bUsed = false;
+			for (UMeshComponent* MeshComponent : MeshPaintingSubsystem->GetPaintableMeshComponents())
 			{
-				FSphere Sphere(BestTraceResult.Location, BrushRadius);
-				if (MeshComponent->GetLocalBounds().GetSphere().TransformBy(MeshComponent->GetComponentTransform()).Intersects(Sphere))
+				if (MeshComponent == BestTraceMeshComponent)
 				{
 					HoveredComponents.FindOrAdd(MeshComponent).Add(i);
 					bUsed = true;
 				}
+				else if (bAllowMultiselect)
+				{
+					FSphere Sphere(BestTraceResult.Location, BrushRadius);
+					if (MeshComponent->GetLocalBounds().GetSphere().TransformBy(MeshComponent->GetComponentTransform()).Intersects(Sphere))
+					{
+						HoveredComponents.FindOrAdd(MeshComponent).Add(i);
+						bUsed = true;
+					}
+				}
 			}
-		}
 
-		if (bUsed)
-		{
-			FVector BrushXAxis, BrushYAxis;
-			BestTraceResult.Normal.FindBestAxisVectors(BrushXAxis, BrushYAxis);
-			// Display settings
-			const float VisualBiasDistance = 0.15f;
-			const FVector BrushVisualPosition = BestTraceResult.Location + BestTraceResult.Normal * VisualBiasDistance;
+			if (bUsed)
+			{
+				FVector BrushXAxis, BrushYAxis;
+				BestTraceResult.Normal.FindBestAxisVectors(BrushXAxis, BrushYAxis);
+				// Display settings
+				const float VisualBiasDistance = 0.15f;
+				const FVector BrushVisualPosition = BestTraceResult.Location + BestTraceResult.Normal * VisualBiasDistance;
 
-			const FLinearColor PaintColor = TextureProperties->PaintColor;
-			const FLinearColor EraseColor = TextureProperties->EraseColor;
+				const FLinearColor PaintColor = TextureProperties->PaintColor;
+				const FLinearColor EraseColor = TextureProperties->EraseColor;
 
-			// NOTE: We square the brush strength to maximize slider precision in the low range
-			const float BrushStrength = BrushProperties->BrushStrength *  BrushProperties->BrushStrength * InStrengthScale;
+				// NOTE: We square the brush strength to maximize slider precision in the low range
+				const float BrushStrength = BrushProperties->BrushStrength * BrushProperties->BrushStrength * InStrengthScale;
 
-			const float BrushDepth = BrushRadius;
+				const float BrushDepth = BrushRadius;
 
-			// Mesh paint settings
-			FMeshPaintParameters& Params = PaintRayResults[i].Params;
-			{
-				Params.PaintAction = PaintAction;
-				Params.BrushPosition = BestTraceResult.Location;
-				Params.BrushNormal = BestTraceResult.Normal;
-				Params.BrushColor = bIsPainting ? PaintColor : EraseColor;
-				Params.SquaredBrushRadius = BrushRadius * BrushRadius;
-				Params.BrushRadialFalloffRange = BrushProperties->BrushFalloffAmount * BrushRadius;
-				Params.InnerBrushRadius = BrushRadius - Params.BrushRadialFalloffRange;
-				Params.BrushDepth = BrushDepth;
-				Params.BrushDepthFalloffRange = BrushProperties->BrushFalloffAmount * BrushDepth;
-				Params.InnerBrushDepth = BrushDepth - Params.BrushDepthFalloffRange;
-				Params.BrushStrength = BrushStrength;
-				Params.BrushToWorldMatrix = FMatrix(BrushXAxis, BrushYAxis, Params.BrushNormal, Params.BrushPosition);
-				Params.InverseBrushToWorldMatrix = Params.BrushToWorldMatrix.InverseFast();
- 				Params.bWriteRed = TextureProperties->bWriteRed;
- 				Params.bWriteGreen = TextureProperties->bWriteGreen;
-				Params.bWriteBlue = TextureProperties->bWriteBlue;
-				Params.bWriteAlpha = TextureProperties->bWriteAlpha;
-				FVector BrushSpaceVertexPosition = Params.InverseBrushToWorldMatrix.TransformVector(FVector4(Params.BrushPosition, 1.0f));
-				Params.BrushPosition2D = FVector2f(BrushSpaceVertexPosition.X, BrushSpaceVertexPosition.Y);
+				// Mesh paint settings
+				FMeshPaintParameters& Params = PaintRayResults[i].Params;
+				{
+					Params.PaintAction = PaintAction;
+					Params.BrushPosition = BestTraceResult.Location;
+					Params.BrushNormal = BestTraceResult.Normal;
+					Params.BrushColor = bIsPainting ? PaintColor : EraseColor;
+					Params.SquaredBrushRadius = BrushRadius * BrushRadius;
+					Params.BrushRadialFalloffRange = BrushProperties->BrushFalloffAmount * BrushRadius;
+					Params.InnerBrushRadius = BrushRadius - Params.BrushRadialFalloffRange;
+					Params.BrushDepth = BrushDepth;
+					Params.BrushDepthFalloffRange = BrushProperties->BrushFalloffAmount * BrushDepth;
+					Params.InnerBrushDepth = BrushDepth - Params.BrushDepthFalloffRange;
+					Params.BrushStrength = BrushStrength;
+					Params.BrushToWorldMatrix = FMatrix(BrushXAxis, BrushYAxis, Params.BrushNormal, Params.BrushPosition);
+					Params.InverseBrushToWorldMatrix = Params.BrushToWorldMatrix.InverseFast();
+					Params.bWriteRed = TextureProperties->bWriteRed;
+					Params.bWriteGreen = TextureProperties->bWriteGreen;
+					Params.bWriteBlue = TextureProperties->bWriteBlue;
+					Params.bWriteAlpha = TextureProperties->bWriteAlpha;
+					FVector BrushSpaceVertexPosition = Params.InverseBrushToWorldMatrix.TransformVector(FVector4(Params.BrushPosition, 1.0f));
+					Params.BrushPosition2D = FVector2f(BrushSpaceVertexPosition.X, BrushSpaceVertexPosition.Y);
+				}
 			}
 		}
-	}
-
-	if (HoveredComponents.Num() > 0)
-	{
-		if (bArePainting == false)
-		{
-			bArePainting = true;
-		}
 
-		// Iterate over the selected meshes under the cursor and paint them!
-		for (auto& Entry : HoveredComponents)
+		if (HoveredComponents.Num() > 0)
 		{
-			UMeshComponent* HoveredComponent = Entry.Key;
-			TArray<int32>& PaintRayResultIds = Entry.Value;
-			
-			TSharedPtr<IMeshPaintComponentAdapter> MeshAdapter = MeshPaintingSubsystem->GetAdapterForComponent(HoveredComponent);
-			if (!IsMeshAdapterSupported(MeshAdapter))
+			if (bArePainting == false)
 			{
-				continue;
-			}
-				
-			const int32 UVChannel = GetSelectedUVChannel(HoveredComponent);
-			if (UVChannel >= MeshAdapter->GetNumUVChannels())
-			{
-				continue;
+				bArePainting = true;
 			}
 
-			TArray<UTexture const*> Textures;
-			const UTexture2D* TargetTexture2D = GetSelectedPaintTexture(HoveredComponent);
-			if (TargetTexture2D == nullptr)
+			// Iterate over the selected meshes under the cursor and paint them!
+			for (auto& Entry : HoveredComponents)
 			{
-				continue;
-			}
+				UMeshComponent* HoveredComponent = Entry.Key;
+				TArray<int32>& PaintRayResultIds = Entry.Value;
+			
+				TSharedPtr<IMeshPaintComponentAdapter> MeshAdapter = MeshPaintingSubsystem->GetAdapterForComponent(HoveredComponent);
+				if (!IsMeshAdapterSupported(MeshAdapter))
+				{
+					continue;
+				}
+				
+				const int32 UVChannel = GetSelectedUVChannel(HoveredComponent);
+				if (UVChannel >= MeshAdapter->GetNumUVChannels())
+				{
+					continue;
+				}
+
+				TArray<UTexture const*> Textures;
+				const UTexture2D* TargetTexture2D = GetSelectedPaintTexture(HoveredComponent);
+				if (TargetTexture2D == nullptr)
+				{
+					continue;
+				}
 					
-			Textures.Add(TargetTexture2D);
+				Textures.Add(TargetTexture2D);
 
-			FPaintTexture2DData* TextureData = GetPaintTargetData(TargetTexture2D);
-			if (TextureData)
-			{
-				Textures.Add(TextureData->PaintRenderTargetTexture);
-			}
+				FPaintTexture2DData* TextureData = GetPaintTargetData(TargetTexture2D);
+				if (TextureData)
+				{
+					Textures.Add(TextureData->PaintRenderTargetTexture);
+				}
 
-			TArray<FTexturePaintMeshSectionInfo> MaterialSections;
-			UTexturePaintToolset::RetrieveMeshSectionsForTextures(HoveredComponent, 0/*CachedLODIndex*/, Textures, MaterialSections);
+				TArray<FTexturePaintMeshSectionInfo> MaterialSections;
+				UTexturePaintToolset::RetrieveMeshSectionsForTextures(HoveredComponent, 0/*CachedLODIndex*/, Textures, MaterialSections);
 
-			bool bPaintApplied = false;
-			TArray<FTexturePaintTriangleInfo> TrianglePaintInfoArray;
-			for (int32 PaintRayResultId : PaintRayResultIds)
-			{
-				const FVector& BestTraceResultLocation = PaintRayResults[PaintRayResultId].BestTraceResult.Location;
-				FViewCameraState CameraState;
-				GetToolManager()->GetContextQueriesAPI()->GetCurrentViewState(CameraState);
-				bPaintApplied |= MeshPaintingSubsystem->ApplyPerTrianglePaintAction(MeshAdapter.Get(), CameraState.Position, BestTraceResultLocation, BrushProperties, FPerTrianglePaintAction::CreateUObject(this, &UMeshTexturePaintingTool::GatherTextureTriangles, &TrianglePaintInfoArray, &MaterialSections, UVChannel), TextureProperties->bOnlyFrontFacingTriangles);
-				break;
-			}
+				bool bPaintApplied = false;
+				TArray<FTexturePaintTriangleInfo> TrianglePaintInfoArray;
+				if (PaintRayResultIds.Num() > 0)
+				{
+					const int32 PaintRayResultId = PaintRayResultIds[0];
+					const FVector& BestTraceResultLocation = PaintRayResults[PaintRayResultId].BestTraceResult.Location;
+					FViewCameraState CameraState;
+					GetToolManager()->GetContextQueriesAPI()->GetCurrentViewState(CameraState);
+					bPaintApplied |= MeshPaintingSubsystem->ApplyPerTrianglePaintAction(MeshAdapter.Get(), CameraState.Position, BestTraceResultLocation, BrushProperties, FPerTrianglePaintAction::CreateUObject(this, &UMeshTexturePaintingTool::GatherTextureTriangles, &TrianglePaintInfoArray, &MaterialSections, UVChannel), TextureProperties->bOnlyFrontFacingTriangles);
+					break;
+				}
 
-			if (!bPaintApplied)
-			{
-				continue;
-			}
+				if (!bPaintApplied)
+				{
+					continue;
+				}
 
-			// Painting textures
-			bAnyPaintApplied = true;
+				// Painting textures
+				bAnyPaintApplied = true;
 
-			UTexture2D* SelectedPaintTexure = GetSelectedPaintTexture(HoveredComponent);
-			if (PaintingTexture2D != nullptr && PaintingTexture2D != SelectedPaintTexure)
-			{
-				// Texture has changed, so finish up with our previous texture
-				FinishPaintingTexture();
-			}
+				UTexture2D* SelectedPaintTexure = GetSelectedPaintTexture(HoveredComponent);
+				if (PaintingTexture2D != nullptr && PaintingTexture2D != SelectedPaintTexure)
+				{
+					// Texture has changed, so finish up with our previous texture
+					FinishPaintingTexture();
+				}
 
-			if (PaintingTexture2D == nullptr)
-			{
-				StartPaintingTexture(HoveredComponent, *MeshAdapter);
-			}
+				if (PaintingTexture2D == nullptr)
+				{
+					StartPaintingTexture(HoveredComponent, *MeshAdapter);
+				}
 
-			if (PaintingTexture2D != nullptr)
-			{
-				for (int32 PaintRayResultId : PaintRayResultIds)
+				if (PaintingTexture2D != nullptr)
 				{
-					FMeshPaintParameters& Params = PaintRayResults[PaintRayResultId].Params;
-					FMeshPaintParameters* LastParams = nullptr;
-					if (LastPaintRayResults.Num() > PaintRayResultId)
+					if (PaintRayResultIds.Num() > 0)
 					{
-						LastParams = &LastPaintRayResults[PaintRayResultId].Params;
+						const int32 PaintRayResultId = PaintRayResultIds[0];
+						FMeshPaintParameters& Params = PaintRayResults[PaintRayResultId].Params;
+						FMeshPaintParameters* LastParams = nullptr;
+						if (LastPaintRayResults.Num() > PaintRayResultId)
+						{
+							LastParams = &LastPaintRayResults[PaintRayResultId].Params;
+						}
+
+						PaintTexture(Params, UVChannel, TrianglePaintInfoArray, HoveredComponent, *MeshAdapter, LastParams);
+						break;
 					}
-
-					PaintTexture(Params, UVChannel, TrianglePaintInfoArray, HoveredComponent, *MeshAdapter, LastParams);
-					break;
 				}
 			}
 		}
 	}
 
-	LastPaintRayResults = PaintRayResults;
+	LastPaintRayResults = MoveTemp(PaintRayResults);
 	return bAnyPaintApplied;
 }
 
