Remove breaks that got left in accidentally as result of failure to fully remove from last changelist different SA fix for loops that would iterate only one time.
#rnx

diff --git a/Engine/Plugins/MeshPainting/Source/MeshPaintingToolset/Private/MeshTexturePaintingTool.cpp b/Engine/Plugins/MeshPainting/Source/MeshPaintingToolset/Private/MeshTexturePaintingTool.cpp
--- a/Engine/Plugins/MeshPainting/Source/MeshPaintingToolset/Private/MeshTexturePaintingTool.cpp
+++ b/Engine/Plugins/MeshPainting/Source/MeshPaintingToolset/Private/MeshTexturePaintingTool.cpp
@@ -517,7 +517,6 @@ bool UMeshTexturePaintingTool::PaintInternal(const TArrayView<TPair<FVector, FVe
 					FViewCameraState CameraState;
 					GetToolManager()->GetContextQueriesAPI()->GetCurrentViewState(CameraState);
 					bPaintApplied |= MeshPaintingSubsystem->ApplyPerTrianglePaintAction(MeshAdapter.Get(), CameraState.Position, BestTraceResultLocation, BrushProperties, FPerTrianglePaintAction::CreateUObject(this, &UMeshTexturePaintingTool::GatherTextureTriangles, &TrianglePaintInfoArray, &MaterialSections, UVChannel), TextureProperties->bOnlyFrontFacingTriangles);
-					break;
 				}
 
 				if (!bPaintApplied)
@@ -553,7 +552,6 @@ bool UMeshTexturePaintingTool::PaintInternal(const TArrayView<TPair<FVector, FVe
 						}
 
 						PaintTexture(Params, UVChannel, TrianglePaintInfoArray, HoveredComponent, *MeshAdapter, LastParams);
-						break;
 					}
 				}
 			}
