Mesh texture painting: Override visual for texture being painted in all materials instead of only the materials found on the painted components.
This is usful for the case where we are sharing a paint texture across many items in the world.
#rnx

diff --git a/Engine/Plugins/MeshPainting/Source/MeshPaintingToolset/Private/IMeshPaintComponentAdapter.cpp b/Engine/Plugins/MeshPainting/Source/MeshPaintingToolset/Private/IMeshPaintComponentAdapter.cpp
--- a/Engine/Plugins/MeshPainting/Source/MeshPaintingToolset/Private/IMeshPaintComponentAdapter.cpp
+++ b/Engine/Plugins/MeshPainting/Source/MeshPaintingToolset/Private/IMeshPaintComponentAdapter.cpp
@@ -12,6 +12,7 @@
 #include "Materials/MaterialExpressionTextureSampleParameter.h"
 #include "Materials/MaterialInstance.h"
 #include "MeshPaintingToolsetTypes.h"
+#include "ObjectCacheContext.h"
 #include "TexturePaintToolset.h"
 #include "UObject/UObjectGlobals.h"
 #include "UObject/UnrealType.h"
@@ -425,30 +426,26 @@ namespace UE::MeshPaintingToolset
 		return *this;
 	}
 
-	void FDefaultTextureOverride::ApplyOrRemoveTextureOverride(UMeshComponent* InMeshComponent, const UTexture* SourceTexture, UTexture* OverrideTexture) const
+	void FDefaultTextureOverride::ApplyOrRemoveTextureOverride(UMeshComponent* InMeshComponent, UTexture* SourceTexture, UTexture* OverrideTexture) const
 	{
 		check(IsInGameThread());
 
 		const ERHIFeatureLevel::Type FeatureLevel = InMeshComponent->GetWorld()->GetFeatureLevel();
 
-		// Check all the materials on the mesh to see if the user texture is there
-		int32 MaterialIndex = 0;
-		UMaterialInterface* MaterialToCheck = InMeshComponent->GetMaterial(MaterialIndex);
-		while (MaterialToCheck != nullptr)
+		// Applying the texture override to all materials shows live update throughout the scene not just on painted component.
+		FObjectCacheContextScope ObjectCacheScope;
+		for (UMaterialInterface* Material : ObjectCacheScope.GetContext().GetMaterialsAffectedByTexture(SourceTexture))
 		{
 			if (!OverrideTexture)
 			{
-				// Unregister for all materials. This will no affect the material that weren't overridden by this instance
-				Private::FGlobalTextureOverrideState::RemoveMaterialOverride(this, MaterialToCheck, SourceTexture);
+				// Unregister for all materials. This will not affect the materials that weren't overridden by this instance
+				Private::FGlobalTextureOverrideState::RemoveMaterialOverride(this, Material, SourceTexture);
 			}
 			else
 			{
-				// Keep track of the material overridden
-				Private::FGlobalTextureOverrideState::RegisterMaterialOverride(this, MaterialToCheck, SourceTexture, OverrideTexture, FeatureLevel);
+				// Override and register to track modified materials.
+				Private::FGlobalTextureOverrideState::RegisterMaterialOverride(this, Material, SourceTexture, OverrideTexture, FeatureLevel);
 			}
-
-			++MaterialIndex;
-			MaterialToCheck = InMeshComponent->GetMaterial(MaterialIndex);
 		}
 
 		// Check to see if the source texture is the special mesh paint texture on the component.
@@ -465,5 +462,4 @@ namespace UE::MeshPaintingToolset
 
 		Private::FGlobalTextureOverrideState::FreeOverrideOwnership(this);
 	}
-
 }
diff --git a/Engine/Plugins/MeshPainting/Source/MeshPaintingToolset/Public/IMeshPaintComponentAdapter.h b/Engine/Plugins/MeshPainting/Source/MeshPaintingToolset/Public/IMeshPaintComponentAdapter.h
--- a/Engine/Plugins/MeshPainting/Source/MeshPaintingToolset/Public/IMeshPaintComponentAdapter.h
+++ b/Engine/Plugins/MeshPainting/Source/MeshPaintingToolset/Public/IMeshPaintComponentAdapter.h
@@ -126,7 +126,7 @@ namespace UE::MeshPaintingToolset
 	class FDefaultTextureOverride final
 	{
 	public:
-		UE_API void ApplyOrRemoveTextureOverride(UMeshComponent* InMeshComponent, const UTexture* SourceTexture, UTexture* OverrideTexture) const;
+		UE_API void ApplyOrRemoveTextureOverride(UMeshComponent* InMeshComponent, UTexture* SourceTexture, UTexture* OverrideTexture) const;
 
 		UE_API ~FDefaultTextureOverride();
 
