Mesh paint fixes.
Recache materials correctly after paint material overrides are applied. We relied on a call to RecacheUniformExpressions() inside UMaterialInstance::OverrideTexture() bit this isn't enough.
Track when a paint texture is used in an RVT material and when it does flush RVTs during painting.
#rnx

diff --git a/Engine/Plugins/MeshPainting/Source/MeshPaintingToolset/Private/IMeshPaintComponentAdapter.cpp b/Engine/Plugins/MeshPainting/Source/MeshPaintingToolset/Private/IMeshPaintComponentAdapter.cpp
--- a/Engine/Plugins/MeshPainting/Source/MeshPaintingToolset/Private/IMeshPaintComponentAdapter.cpp
+++ b/Engine/Plugins/MeshPainting/Source/MeshPaintingToolset/Private/IMeshPaintComponentAdapter.cpp
@@ -206,7 +206,7 @@ namespace UE::MeshPaintingToolset
 				}
 			}
 
-			static void RegisterMaterialOverride(const FDefaultTextureOverride* Requester, UMaterialInterface* Material, const UTexture* SourceTexture, UTexture* OverrrideTexture, const ERHIFeatureLevel::Type FeatureLevel)
+			static void RegisterMaterialOverride(const FDefaultTextureOverride* Requester, UMaterialInterface* Material, const UTexture* SourceTexture, UTexture* OverrrideTexture, const ERHIFeatureLevel::Type FeatureLevel, bool& bOutUpdated)
 			{
 				TSet<FOverrideKey, DefaultKeyFuncs<FOverrideKey>, TInlineSetAllocator<2>>& Overrides = DefaultTextureOverrideToOverrides.FindOrAdd(Requester);
 				
@@ -228,6 +228,7 @@ namespace UE::MeshPaintingToolset
 					for (const ERHIFeatureLevel::Type LevelToUpdate : OverrideData.OverridenFeatureLevels)
 					{
 						Material->OverrideTexture(SourceTexture, OverrrideTexture, LevelToUpdate);
+						bOutUpdated = true;
 					}
 
 					AddMaterialTracking(Pair);
@@ -241,7 +242,7 @@ namespace UE::MeshPaintingToolset
 				OverrideData.Count += 1;
 			}
 
-			static void RemoveMaterialOverride(const FDefaultTextureOverride* Requester, UMaterialInterface* Material, const UTexture* SourceTexture)
+			static void RemoveMaterialOverride(const FDefaultTextureOverride* Requester, UMaterialInterface* Material, const UTexture* SourceTexture, bool& bOutUpdated)
 			{
 				if (TSet<FOverrideKey, DefaultKeyFuncs<FOverrideKey>, TInlineSetAllocator<2>>* Override = DefaultTextureOverrideToOverrides.Find(Requester))
 				{
@@ -258,6 +259,7 @@ namespace UE::MeshPaintingToolset
 								for (const ERHIFeatureLevel::Type FeatureLevel : OverrideData->OverridenFeatureLevels)
 								{
 									Material->OverrideTexture(SourceTexture, nullptr, FeatureLevel);
+									bOutUpdated = true;
 								}
 
 								OverridesData.RemoveByHash(PairHash, Pair);
@@ -432,19 +434,28 @@ namespace UE::MeshPaintingToolset
 
 		const ERHIFeatureLevel::Type FeatureLevel = InMeshComponent->GetWorld()->GetFeatureLevel();
 
+		FMaterialUpdateContext MaterialUpdateContext;
+
 		// Applying the texture override to all materials shows live update throughout the scene not just on painted component.
 		FObjectCacheContextScope ObjectCacheScope;
 		for (UMaterialInterface* Material : ObjectCacheScope.GetContext().GetMaterialsAffectedByTexture(SourceTexture))
 		{
+			bool bMaterialUpdated = false;
+			
 			if (!OverrideTexture)
 			{
 				// Unregister for all materials. This will not affect the materials that weren't overridden by this instance
-				Private::FGlobalTextureOverrideState::RemoveMaterialOverride(this, Material, SourceTexture);
+				Private::FGlobalTextureOverrideState::RemoveMaterialOverride(this, Material, SourceTexture, bMaterialUpdated);
 			}
 			else
 			{
 				// Override and register to track modified materials.
-				Private::FGlobalTextureOverrideState::RegisterMaterialOverride(this, Material, SourceTexture, OverrideTexture, FeatureLevel);
+				Private::FGlobalTextureOverrideState::RegisterMaterialOverride(this, Material, SourceTexture, OverrideTexture, FeatureLevel, bMaterialUpdated);
+			}
+
+			if (bMaterialUpdated)
+			{
+				MaterialUpdateContext.AddMaterialInterface(Material);
 			}
 		}
 
diff --git a/Engine/Plugins/MeshPainting/Source/MeshPaintingToolset/Private/MeshTexturePaintingTool.cpp b/Engine/Plugins/MeshPainting/Source/MeshPaintingToolset/Private/MeshTexturePaintingTool.cpp
--- a/Engine/Plugins/MeshPainting/Source/MeshPaintingToolset/Private/MeshTexturePaintingTool.cpp
+++ b/Engine/Plugins/MeshPainting/Source/MeshPaintingToolset/Private/MeshTexturePaintingTool.cpp
@@ -3,31 +3,34 @@
 #include "MeshTexturePaintingTool.h"
 
 #include "AssetRegistry/AssetData.h"
-#include "InteractiveToolManager.h"
+#include "BaseGizmos/BrushStampIndicator.h"
+#include "CanvasItem.h"
+#include "CanvasTypes.h"
 #include "Components/MeshComponent.h"
+#include "Components/RuntimeVirtualTextureComponent.h"
+#include "Editor/EditorEngine.h"
+#include "Editor/TransBuffer.h"
+#include "Editor/UnrealEdEngine.h"
+#include "Engine/Texture2D.h"
 #include "Engine/TextureRenderTarget2D.h"
 #include "IMeshPaintComponentAdapter.h"
+#include "InteractiveToolManager.h"
 #include "Materials/MaterialInterface.h"
-#include "ToolDataVisualizer.h"
-#include "Engine/Texture2D.h"
-#include "CanvasItem.h"
-#include "CanvasTypes.h"
 #include "MeshPaintHelpers.h"
 #include "MeshVertexPaintingTool.h"
-#include "ScopedTransaction.h"
-#include "BaseGizmos/BrushStampIndicator.h"
+#include "ObjectCacheContext.h"
 #include "RenderingThread.h"
 #include "RHIUtilities.h"
+#include "ScopedTransaction.h"
 #include "TextureCompiler.h"
 #include "TexturePaintToolset.h"
 #include "TextureResource.h"
-#include "Editor/TransBuffer.h"
-#include "Editor/UnrealEdEngine.h"
+#include "ToolDataVisualizer.h"
 #include "VT/VirtualTextureAdapter.h"
 #include "VT/VirtualTextureBuildSettings.h"
 #include "UnrealEdGlobals.h"
+#include "UObject/UObjectIterator.h"
 
-#include "Editor/EditorEngine.h"
 extern UNREALED_API class UEditorEngine* GEditor;
 
 #include UE_INLINE_GENERATED_CPP_BY_NAME(MeshTexturePaintingTool)
@@ -678,6 +681,7 @@ void UMeshTexturePaintingTool::FinishPainting()
 {
 	PaintingTransaction.Reset();
 	bArePainting = false;
+	bRequiresRuntimeVirtualTextureUpdates = false;
 }
 
 FPaintTexture2DData* UMeshTexturePaintingTool::GetPaintTargetData(const UTexture2D* InTexture)
@@ -875,6 +879,16 @@ void UMeshTexturePaintingTool::StartPaintingTexture(UMeshComponent* InMeshCompon
 			UTexturePaintToolset::SetupInitialRenderTargetData(TextureProperties->PaintBrush, TextureData->PaintBrushRenderTargetTexture);
 		}
 	}
+
+	{
+		// Set flag for whether to update runtime virtual textures while painting.
+		bRequiresRuntimeVirtualTextureUpdates = false;
+		FObjectCacheContextScope ObjectCacheScope;
+		for (UMaterialInterface* Material : ObjectCacheScope.GetContext().GetMaterialsAffectedByTexture(Texture2D))
+		{
+			bRequiresRuntimeVirtualTextureUpdates |= Material->WritesToRuntimeVirtualTexture();
+		}
+	}
 }
 
 void UMeshTexturePaintingTool::PaintTexture(FMeshPaintParameters& InParams, int32 UVChannel, TArray<FTexturePaintTriangleInfo>& InInfluencedTriangles, UMeshComponent* MeshComponent, const IMeshPaintComponentAdapter& GeometryInfo, FMeshPaintParameters* LastParams)
@@ -1214,6 +1228,15 @@ void UMeshTexturePaintingTool::PaintTexture(FMeshPaintParameters& InParams, int3
 	{
 		TextureData->PaintRenderTargetTextureAdapter->Flush(FBox2f(FVector2f(0, 0), FVector2f(1, 1)));
 	}
+
+	// Need to invalidate runtime virtual textures if paint texture is used in a material that writes to one.
+	if (bRequiresRuntimeVirtualTextureUpdates)
+	{
+		for (TObjectIterator<URuntimeVirtualTextureComponent> It(RF_ClassDefaultObject, true, EInternalObjectFlags::Garbage); It; ++It)
+		{
+			It->Invalidate(It->Bounds, EVTInvalidatePriority::High);
+		}
+	}
 }
 
 
diff --git a/Engine/Plugins/MeshPainting/Source/MeshPaintingToolset/Public/MeshTexturePaintingTool.h b/Engine/Plugins/MeshPainting/Source/MeshPaintingToolset/Public/MeshTexturePaintingTool.h
--- a/Engine/Plugins/MeshPainting/Source/MeshPaintingToolset/Public/MeshTexturePaintingTool.h
+++ b/Engine/Plugins/MeshPainting/Source/MeshPaintingToolset/Public/MeshTexturePaintingTool.h
@@ -218,6 +218,7 @@ private:
 	bool bStampPending = false;
 	bool bInDrag = false;
 	bool bRequestPaintBucketFill = false;
+	bool bRequiresRuntimeVirtualTextureUpdates = false;
 
 	bool bCachedClickRay = false;
 	FRay PendingStampRay;
