MeshPaint: Share a single FMaterialUpdateContext when overriding textures across many components.
Improves performance when many components are selected.
#rnx

diff --git a/Engine/Plugins/MeshPainting/Source/MeshPaintingToolset/Private/IMeshPaintComponentAdapter.cpp b/Engine/Plugins/MeshPainting/Source/MeshPaintingToolset/Private/IMeshPaintComponentAdapter.cpp
--- a/Engine/Plugins/MeshPainting/Source/MeshPaintingToolset/Private/IMeshPaintComponentAdapter.cpp
+++ b/Engine/Plugins/MeshPainting/Source/MeshPaintingToolset/Private/IMeshPaintComponentAdapter.cpp
@@ -21,27 +21,6 @@
 //////////////////////////////////////////////////////////////////////////
 // IMeshPaintGeometryAdapter
 
-void IMeshPaintComponentAdapter::DefaultApplyOrRemoveTextureOverride(UMeshComponent* InMeshComponent, UTexture* SourceTexture, UTexture* OverrideTexture)
-{
-	// Legacy implementation for the deprecation duration. For a valid implementation look at deprecation comment
-
-	const ERHIFeatureLevel::Type FeatureLevel = InMeshComponent->GetWorld()->GetFeatureLevel();
-
-	// Check all the materials on the mesh to see if the user texture is there
-	int32 MaterialIndex = 0;
-	UMaterialInterface* MaterialToCheck = InMeshComponent->GetMaterial(MaterialIndex);
-	while (MaterialToCheck != nullptr)
-	{
-		if (!OverrideTexture || DoesMaterialUseTexture(MaterialToCheck, SourceTexture))
-		{
-			MaterialToCheck->OverrideTexture(SourceTexture, OverrideTexture, FeatureLevel);
-		}
-
-		++MaterialIndex;
-		MaterialToCheck = InMeshComponent->GetMaterial(MaterialIndex);
-	}
-}
-
 static bool IsTextureSuitableForTexturePainting(const TWeakObjectPtr<UTexture> TexturePtr)
 {
 	return (TexturePtr.Get() != nullptr &&
@@ -167,6 +146,12 @@ void IMeshPaintComponentAdapter::DefaultQueryPaintableTextures(int32 MaterialInd
 	}
 }
 
+void IMeshPaintComponentAdapter::ApplyOrRemoveTextureOverride(UTexture* SourceTexture, UTexture* OverrideTexture) const
+{
+	FMaterialUpdateContext MaterialUpdateContext;
+	ApplyOrRemoveTextureOverride(SourceTexture, OverrideTexture, MaterialUpdateContext);
+}
+
 namespace UE::MeshPaintingToolset
 {
 	namespace Private
@@ -428,14 +413,12 @@ namespace UE::MeshPaintingToolset
 		return *this;
 	}
 
-	void FDefaultTextureOverride::ApplyOrRemoveTextureOverride(UMeshComponent* InMeshComponent, UTexture* SourceTexture, UTexture* OverrideTexture) const
+	void FDefaultTextureOverride::ApplyOrRemoveTextureOverride(UMeshComponent* InMeshComponent, UTexture* SourceTexture, UTexture* OverrideTexture, FMaterialUpdateContext& MaterialUpdateContext) const
 	{
 		check(IsInGameThread());
 
 		const ERHIFeatureLevel::Type FeatureLevel = InMeshComponent->GetWorld()->GetFeatureLevel();
 
-		FMaterialUpdateContext MaterialUpdateContext;
-
 		// Applying the texture override to all materials shows live update throughout the scene not just on painted component.
 		FObjectCacheContextScope ObjectCacheScope;
 		for (UMaterialInterface* Material : ObjectCacheScope.GetContext().GetMaterialsAffectedByTexture(SourceTexture))
@@ -467,6 +450,13 @@ namespace UE::MeshPaintingToolset
 		}
 	}
 
+	void FDefaultTextureOverride::ApplyOrRemoveTextureOverride(UMeshComponent* InMeshComponent, UTexture* SourceTexture, UTexture* OverrideTexture) const
+	{
+		FMaterialUpdateContext MaterialUpdateContext;
+		ApplyOrRemoveTextureOverride(InMeshComponent, SourceTexture, OverrideTexture, MaterialUpdateContext);
+	}
+
+
 	FDefaultTextureOverride::~FDefaultTextureOverride()
 	{
 		check(IsInGameThread());
diff --git a/Engine/Plugins/MeshPainting/Source/MeshPaintingToolset/Private/MeshPaintGeometryCollectionAdapter.cpp b/Engine/Plugins/MeshPainting/Source/MeshPaintingToolset/Private/MeshPaintGeometryCollectionAdapter.cpp
--- a/Engine/Plugins/MeshPainting/Source/MeshPaintingToolset/Private/MeshPaintGeometryCollectionAdapter.cpp
+++ b/Engine/Plugins/MeshPainting/Source/MeshPaintingToolset/Private/MeshPaintGeometryCollectionAdapter.cpp
@@ -248,11 +248,11 @@ void FMeshPaintGeometryCollectionComponentAdapter::QueryPaintableTextures(int32
 	}
 }
 
-void FMeshPaintGeometryCollectionComponentAdapter::ApplyOrRemoveTextureOverride(UTexture* SourceTexture, UTexture* OverrideTexture) const
+void FMeshPaintGeometryCollectionComponentAdapter::ApplyOrRemoveTextureOverride(UTexture* SourceTexture, UTexture* OverrideTexture, FMaterialUpdateContext& MaterialUpdateContext) const
 {
 	if (GeometryCollectionComponent.IsValid())
 	{
-		TextureOverridesState.ApplyOrRemoveTextureOverride(GeometryCollectionComponent.Get(), SourceTexture, OverrideTexture);
+		TextureOverridesState.ApplyOrRemoveTextureOverride(GeometryCollectionComponent.Get(), SourceTexture, OverrideTexture, MaterialUpdateContext);
 	}
 }
 
diff --git a/Engine/Plugins/MeshPainting/Source/MeshPaintingToolset/Private/MeshPaintHelpers.cpp b/Engine/Plugins/MeshPainting/Source/MeshPaintingToolset/Private/MeshPaintHelpers.cpp
--- a/Engine/Plugins/MeshPainting/Source/MeshPaintingToolset/Private/MeshPaintHelpers.cpp
+++ b/Engine/Plugins/MeshPainting/Source/MeshPaintingToolset/Private/MeshPaintHelpers.cpp
@@ -955,23 +955,6 @@ void UMeshPaintingSubsystem::ForceRenderMeshLOD(UMeshComponent* Component, int32
 	}	
 }
 
-void UMeshPaintingSubsystem::ClearMeshTextureOverrides(const IMeshPaintComponentAdapter& GeometryInfo, UMeshComponent* InMeshComponent)
-{
-	if (InMeshComponent != nullptr)
-	{
-		TArray<UTexture*> UsedTextures;
-		InMeshComponent->GetUsedTextures(/*out*/ UsedTextures, EMaterialQualityLevel::High);
-
-		for (UTexture* Texture : UsedTextures)
-		{
-			if (UTexture2D* Texture2D = Cast<UTexture2D>(Texture))
-			{
-				GeometryInfo.ApplyOrRemoveTextureOverride(Texture2D, nullptr);
-			}
-		}
-	}
-}
-
 void UMeshPaintingSubsystem::ApplyVertexColorsToAllLODs(IMeshPaintComponentAdapter& GeometryInfo, UMeshComponent* InMeshComponent)
 {
 	if (UStaticMeshComponent* StaticMeshComponent = Cast<UStaticMeshComponent>(InMeshComponent))
diff --git a/Engine/Plugins/MeshPainting/Source/MeshPaintingToolset/Private/MeshPaintSkeletalMeshAdapter.cpp b/Engine/Plugins/MeshPainting/Source/MeshPaintingToolset/Private/MeshPaintSkeletalMeshAdapter.cpp
--- a/Engine/Plugins/MeshPainting/Source/MeshPaintingToolset/Private/MeshPaintSkeletalMeshAdapter.cpp
+++ b/Engine/Plugins/MeshPainting/Source/MeshPaintingToolset/Private/MeshPaintSkeletalMeshAdapter.cpp
@@ -348,11 +348,11 @@ void FMeshPaintSkeletalMeshComponentAdapter::QueryPaintableTextures(int32 Materi
 	}
 }
 
-void FMeshPaintSkeletalMeshComponentAdapter::ApplyOrRemoveTextureOverride(UTexture* SourceTexture, UTexture* OverrideTexture) const
+void FMeshPaintSkeletalMeshComponentAdapter::ApplyOrRemoveTextureOverride(UTexture* SourceTexture, UTexture* OverrideTexture, FMaterialUpdateContext& MaterialUpdateContext) const
 {
 	if (SkeletalMeshComponent.IsValid())
 	{
-		TextureOverridesState.ApplyOrRemoveTextureOverride(SkeletalMeshComponent.Get(), SourceTexture, OverrideTexture);
+		TextureOverridesState.ApplyOrRemoveTextureOverride(SkeletalMeshComponent.Get(), SourceTexture, OverrideTexture, MaterialUpdateContext);
 	}
 }
 
diff --git a/Engine/Plugins/MeshPainting/Source/MeshPaintingToolset/Private/MeshPaintStaticMeshAdapter.cpp b/Engine/Plugins/MeshPainting/Source/MeshPaintingToolset/Private/MeshPaintStaticMeshAdapter.cpp
--- a/Engine/Plugins/MeshPainting/Source/MeshPaintingToolset/Private/MeshPaintStaticMeshAdapter.cpp
+++ b/Engine/Plugins/MeshPainting/Source/MeshPaintingToolset/Private/MeshPaintStaticMeshAdapter.cpp
@@ -253,11 +253,11 @@ void FMeshPaintStaticMeshComponentAdapter::QueryPaintableTextures(int32 Material
 	}
 }
 
-void FMeshPaintStaticMeshComponentAdapter::ApplyOrRemoveTextureOverride(UTexture* SourceTexture, UTexture* OverrideTexture) const
+void FMeshPaintStaticMeshComponentAdapter::ApplyOrRemoveTextureOverride(UTexture* SourceTexture, UTexture* OverrideTexture, FMaterialUpdateContext& MaterialUpdateContext) const
 {
 	if (StaticMeshComponent.IsValid())
 	{
-		TextureOverridesState.ApplyOrRemoveTextureOverride(StaticMeshComponent.Get(), SourceTexture, OverrideTexture);
+		TextureOverridesState.ApplyOrRemoveTextureOverride(StaticMeshComponent.Get(), SourceTexture, OverrideTexture, MaterialUpdateContext);
 	}
 }
 
diff --git a/Engine/Plugins/MeshPainting/Source/MeshPaintingToolset/Private/MeshTexturePaintingTool.cpp b/Engine/Plugins/MeshPainting/Source/MeshPaintingToolset/Private/MeshTexturePaintingTool.cpp
--- a/Engine/Plugins/MeshPainting/Source/MeshPaintingToolset/Private/MeshTexturePaintingTool.cpp
+++ b/Engine/Plugins/MeshPainting/Source/MeshPaintingToolset/Private/MeshTexturePaintingTool.cpp
@@ -16,6 +16,7 @@
 #include "IMeshPaintComponentAdapter.h"
 #include "InteractiveToolManager.h"
 #include "Materials/MaterialInterface.h"
+#include "MaterialShared.h"
 #include "MeshPaintHelpers.h"
 #include "MeshVertexPaintingTool.h"
 #include "ObjectCacheContext.h"
@@ -573,7 +574,7 @@ static UTexture* GetTextureForMaterialOverride(FPaintTexture2DData const& Textur
 	return RenderTargetAdapter != nullptr ? RenderTargetAdapter : RenderTarget;
 }
 
-void UMeshTexturePaintingTool::AddTextureOverrideToComponent(FPaintTexture2DData& TextureData, UMeshComponent* MeshComponent, const IMeshPaintComponentAdapter* MeshPaintAdapter)
+void UMeshTexturePaintingTool::AddTextureOverrideToComponent(FPaintTexture2DData& TextureData, UMeshComponent* MeshComponent, const IMeshPaintComponentAdapter* MeshPaintAdapter, FMaterialUpdateContext& MaterialUpdateContext)
 {
 	if (MeshComponent && MeshPaintAdapter)
 	{
@@ -581,7 +582,7 @@ void UMeshTexturePaintingTool::AddTextureOverrideToComponent(FPaintTexture2DData
 		{
 			TextureData.TextureOverrideComponents.AddUnique(MeshComponent);
 
-			MeshPaintAdapter->ApplyOrRemoveTextureOverride(TextureData.PaintingTexture2D, GetTextureForMaterialOverride(TextureData));
+			MeshPaintAdapter->ApplyOrRemoveTextureOverride(TextureData.PaintingTexture2D, GetTextureForMaterialOverride(TextureData), MaterialUpdateContext);
 		}
 	}
 }
@@ -1295,6 +1296,8 @@ void UMeshTexturePaintingTool::ClearAllTextureOverrides()
 {
 	if (UMeshPaintingSubsystem* MeshPaintingSubsystem = GEngine->GetEngineSubsystem<UMeshPaintingSubsystem>())
 	{
+		FMaterialUpdateContext MaterialUpdateContext;
+
 		/** Remove all texture overrides which are currently stored and active */
 		for (decltype(PaintTargetData)::TIterator It(PaintTargetData); It; ++It)
 		{
@@ -1304,7 +1307,7 @@ void UMeshTexturePaintingTool::ClearAllTextureOverrides()
 			{
 				if (TSharedPtr<IMeshPaintComponentAdapter> PaintAdapter = MeshPaintingSubsystem->GetAdapterForComponent(MeshComponent))
 				{
-					PaintAdapter->ApplyOrRemoveTextureOverride(TextureData->PaintingTexture2D, nullptr);
+					PaintAdapter->ApplyOrRemoveTextureOverride(TextureData->PaintingTexture2D, nullptr, MaterialUpdateContext);
 				}
 			}
 	
@@ -1393,10 +1396,11 @@ void UMeshTexturePaintingTool::SetAllTextureOverrides()
 		}
 
 		// Apply the overrides.
+		FMaterialUpdateContext MaterialUpdateContext;
 		for (UMeshComponent* MeshComponent : PaintableMeshComponents)
 		{
 			TSharedPtr<IMeshPaintComponentAdapter> MeshAdapter = MeshPaintingSubsystem->GetAdapterForComponent(MeshComponent);
-			AddTextureOverrideToComponent(*TextureData, MeshComponent, MeshAdapter.Get());
+			AddTextureOverrideToComponent(*TextureData, MeshComponent, MeshAdapter.Get(), MaterialUpdateContext);
 		}
 	}
 }
@@ -1546,7 +1550,7 @@ void UMeshTextureAssetPaintingTool::OnPropertyModified(UObject* PropertySet, FPr
 				break;
 			}
 		}
-
+		
 		// Need to recreate the render target overrides with the newly selected texture.
 		ClearAllTextureOverrides();
 		SetAllTextureOverrides();
diff --git a/Engine/Plugins/MeshPainting/Source/MeshPaintingToolset/Public/IMeshPaintComponentAdapter.h b/Engine/Plugins/MeshPainting/Source/MeshPaintingToolset/Public/IMeshPaintComponentAdapter.h
--- a/Engine/Plugins/MeshPainting/Source/MeshPaintingToolset/Public/IMeshPaintComponentAdapter.h
+++ b/Engine/Plugins/MeshPainting/Source/MeshPaintingToolset/Public/IMeshPaintComponentAdapter.h
@@ -6,6 +6,7 @@
 
 #define UE_API MESHPAINTINGTOOLSET_API
 
+class FMaterialUpdateContext;
 class FReferenceCollector;
 class UMeshComponent;
 class UTexture;
@@ -64,8 +65,11 @@ public:
 	virtual void QueryPaintableTextures(int32 MaterialIndex, int32& OutDefaultIndex, TArray<struct FPaintableTexture>& InOutTextureList) = 0;
 
 	/** Applies or removes an override texture to use while rendering the meshes materials */
-	virtual void ApplyOrRemoveTextureOverride(UTexture* SourceTexture, UTexture* OverrideTexture) const = 0;
-			
+	virtual void ApplyOrRemoveTextureOverride(UTexture* SourceTexture, UTexture* OverrideTexture, FMaterialUpdateContext& MaterialUpdateContext) const = 0;
+	
+	UE_DEPRECATED(5.7, "Use the version that takes a FMaterialUpdateContext.")
+	void ApplyOrRemoveTextureOverride(UTexture* SourceTexture, UTexture* OverrideTexture) const;
+
 	/** Returns the vertices for the current LOD index in the Mesh */
 	virtual const TArray<FVector>& GetMeshVertices() const = 0;
 
@@ -93,10 +97,6 @@ public:
 	/** Post Edit to clear up or update necessary data */
 	virtual void PostEdit() = 0;
 
-	/** Default functionality for applying or remove a texture override */
-	UE_DEPRECATED(5.4, "This method will be removed because its implementation is not able to handle an material instance being edited. Please inherit from or use an UE::MeshPaintingToolset::FDefaultTextureOverride object and use it function ApplyOrRemoveTextureOverride.")
-	static UE_API void DefaultApplyOrRemoveTextureOverride(UMeshComponent* InMeshComponent, UTexture* SourceTexture, UTexture* OverrideTexture);
-
 	/** Default functionality for querying textures from a mesh component which are suitable for texture painting */
 	static UE_API void DefaultQueryPaintableTextures(int32 MaterialIndex, const UMeshComponent* MeshComponent, int32& OutDefaultIndex, TArray<struct FPaintableTexture>& InOutTextureList);
 
@@ -126,6 +126,8 @@ namespace UE::MeshPaintingToolset
 	class FDefaultTextureOverride final
 	{
 	public:
+		UE_API void ApplyOrRemoveTextureOverride(UMeshComponent* InMeshComponent, UTexture* SourceTexture, UTexture* OverrideTexture, FMaterialUpdateContext& MaterialUpdateContext) const;
+		UE_DEPRECATED(5.7, "Use the version that takes a FMaterialUpdateContext.")
 		UE_API void ApplyOrRemoveTextureOverride(UMeshComponent* InMeshComponent, UTexture* SourceTexture, UTexture* OverrideTexture) const;
 
 		UE_API ~FDefaultTextureOverride();
diff --git a/Engine/Plugins/MeshPainting/Source/MeshPaintingToolset/Public/MeshPaintGeometryCollectionAdapter.h b/Engine/Plugins/MeshPainting/Source/MeshPaintingToolset/Public/MeshPaintGeometryCollectionAdapter.h
--- a/Engine/Plugins/MeshPainting/Source/MeshPaintingToolset/Public/MeshPaintGeometryCollectionAdapter.h
+++ b/Engine/Plugins/MeshPainting/Source/MeshPaintingToolset/Public/MeshPaintGeometryCollectionAdapter.h
@@ -36,7 +36,7 @@ public:
 	UE_API virtual int32 GetNumUVChannels() const override;
 	UE_API virtual bool LineTraceComponent(struct FHitResult& OutHit, const FVector Start, const FVector End, const struct FCollisionQueryParams& Params) const override;
 	UE_API virtual void QueryPaintableTextures(int32 MaterialIndex, int32& OutDefaultIndex, TArray<struct FPaintableTexture>& InOutTextureList) override;
-	UE_API virtual void ApplyOrRemoveTextureOverride(UTexture* SourceTexture, UTexture* OverrideTexture) const override;
+	UE_API virtual void ApplyOrRemoveTextureOverride(UTexture* SourceTexture, UTexture* OverrideTexture, FMaterialUpdateContext& MaterialUpdateContext) const override;
 	UE_API virtual void AddReferencedObjects(FReferenceCollector& Collector) override;
 	UE_API virtual void PreEdit() override;
 	UE_API virtual void PostEdit() override;
diff --git a/Engine/Plugins/MeshPainting/Source/MeshPaintingToolset/Public/MeshPaintHelpers.h b/Engine/Plugins/MeshPainting/Source/MeshPaintingToolset/Public/MeshPaintHelpers.h
--- a/Engine/Plugins/MeshPainting/Source/MeshPaintingToolset/Public/MeshPaintHelpers.h
+++ b/Engine/Plugins/MeshPainting/Source/MeshPaintingToolset/Public/MeshPaintHelpers.h
@@ -198,7 +198,7 @@ public:
 	/** Forces the component to render LOD level at LODIndex instead of the view-based LOD level ( X = 0 means do not force the LOD, X > 0 means force the lod to X - 1 ) */
 	UE_API void ForceRenderMeshLOD(UMeshComponent* Component, int32 LODIndex);
 
-	/** Clears all texture overrides for this component. */
+	UE_DEPRECATED(5.7, "No longer supported. Will be a no-op.")
 	UE_API void ClearMeshTextureOverrides(const IMeshPaintComponentAdapter& GeometryInfo, UMeshComponent* InMeshComponent);
 
 	/** Applies vertex color painting found on LOD 0 to all lower LODs. */
diff --git a/Engine/Plugins/MeshPainting/Source/MeshPaintingToolset/Public/MeshPaintSkeletalMeshAdapter.h b/Engine/Plugins/MeshPainting/Source/MeshPaintingToolset/Public/MeshPaintSkeletalMeshAdapter.h
--- a/Engine/Plugins/MeshPainting/Source/MeshPaintingToolset/Public/MeshPaintSkeletalMeshAdapter.h
+++ b/Engine/Plugins/MeshPainting/Source/MeshPaintingToolset/Public/MeshPaintSkeletalMeshAdapter.h
@@ -38,7 +38,7 @@ public:
 	UE_API virtual int32 GetNumUVChannels() const override;
 	UE_API virtual bool LineTraceComponent(struct FHitResult& OutHit, const FVector Start, const FVector End, const struct FCollisionQueryParams& Params) const override;
 	UE_API virtual void QueryPaintableTextures(int32 MaterialIndex, int32& OutDefaultIndex, TArray<struct FPaintableTexture>& InOutTextureList) override;
-	UE_API virtual void ApplyOrRemoveTextureOverride(UTexture* SourceTexture, UTexture* OverrideTexture) const override;
+	UE_API virtual void ApplyOrRemoveTextureOverride(UTexture* SourceTexture, UTexture* OverrideTexture, FMaterialUpdateContext& MaterialUpdateContext) const override;
 	UE_API virtual void AddReferencedObjects(FReferenceCollector& Collector) override;
 	UE_API virtual void PreEdit() override;
 	UE_API virtual void PostEdit() override;
diff --git a/Engine/Plugins/MeshPainting/Source/MeshPaintingToolset/Public/MeshPaintStaticMeshAdapter.h b/Engine/Plugins/MeshPainting/Source/MeshPaintingToolset/Public/MeshPaintStaticMeshAdapter.h
--- a/Engine/Plugins/MeshPainting/Source/MeshPaintingToolset/Public/MeshPaintStaticMeshAdapter.h
+++ b/Engine/Plugins/MeshPainting/Source/MeshPaintingToolset/Public/MeshPaintStaticMeshAdapter.h
@@ -35,7 +35,7 @@ public:
 	virtual bool SupportsVertexPaint() const override { return StaticMeshComponent.IsValid() && StaticMeshComponent->CanMeshPaintVertexColors(); }
 	UE_API virtual bool LineTraceComponent(struct FHitResult& OutHit, const FVector Start, const FVector End, const struct FCollisionQueryParams& Params) const override;	
 	UE_API virtual void QueryPaintableTextures(int32 MaterialIndex, int32& OutDefaultIndex, TArray<struct FPaintableTexture>& InOutTextureList) override;
-	UE_API virtual void ApplyOrRemoveTextureOverride(UTexture* SourceTexture, UTexture* OverrideTexture) const override;
+	UE_API virtual void ApplyOrRemoveTextureOverride(UTexture* SourceTexture, UTexture* OverrideTexture, FMaterialUpdateContext& MaterialUpdateContext) const override;
 	UE_API virtual void AddReferencedObjects(FReferenceCollector& Collector) override;
 	UE_API virtual void GetTextureCoordinate(int32 VertexIndex, int32 ChannelIndex, FVector2D& OutTextureCoordinate) const override;
 	UE_API virtual void GetVertexColor(int32 VertexIndex, FColor& OutColor, bool bInstance = true) const override;
diff --git a/Engine/Plugins/MeshPainting/Source/MeshPaintingToolset/Public/MeshTexturePaintingTool.h b/Engine/Plugins/MeshPainting/Source/MeshPaintingToolset/Public/MeshTexturePaintingTool.h
--- a/Engine/Plugins/MeshPainting/Source/MeshPaintingToolset/Public/MeshTexturePaintingTool.h
+++ b/Engine/Plugins/MeshPainting/Source/MeshPaintingToolset/Public/MeshTexturePaintingTool.h
@@ -12,6 +12,7 @@
 
 enum class EMeshPaintModeAction : uint8;
 enum class EToolShutdownType : uint8;
+class FMaterialUpdateContext;
 class FScopedTransaction;
 class IMeshPaintComponentAdapter;
 class UMeshToolManager;
@@ -177,18 +178,18 @@ protected:
 	virtual bool CanPaintTextureToComponent(UTexture* InTexture, UMeshComponent const* InMeshComponent) const { return false; }
 
 private:
-	UE_API void CacheSelectionData();
-	UE_API void AddTextureOverrideToComponent(FPaintTexture2DData& TextureData, UMeshComponent* MeshComponent, const IMeshPaintComponentAdapter* MeshPaintAdapter = nullptr);
-	UE_API double CalculateTargetEdgeLength(int TargetTriCount);
-	UE_API void StartPaintingTexture(UMeshComponent* InMeshComponent, const IMeshPaintComponentAdapter& GeometryInfo);
-	UE_API void UpdateResult();
-	UE_API void GatherTextureTriangles(IMeshPaintComponentAdapter* Adapter, int32 TriangleIndex, const int32 VertexIndices[3], TArray<FTexturePaintTriangleInfo>* TriangleInfo, TArray<FTexturePaintMeshSectionInfo>* SectionInfos, int32 UVChannelIndex);
-	UE_API bool Paint(const FVector& InRayOrigin, const FVector& InRayDirection);
-	UE_API bool Paint(const TArrayView<TPair<FVector, FVector>>& Rays);
-	UE_API void PaintTexture(FMeshPaintParameters& InParams, int32 UVChannel, TArray<FTexturePaintTriangleInfo>& InInfluencedTriangles, UMeshComponent* MeshComponent, const IMeshPaintComponentAdapter& GeometryInfo, FMeshPaintParameters* LastParams = nullptr);
-	UE_API bool PaintInternal(const TArrayView<TPair<FVector, FVector>>& Rays, EMeshPaintModeAction PaintAction, float PaintStrength);
-	UE_API void FinishPaintingTexture();
-	UE_API void FinishPainting();
+	void CacheSelectionData();
+	void AddTextureOverrideToComponent(FPaintTexture2DData& TextureData, UMeshComponent* MeshComponent, const IMeshPaintComponentAdapter* MeshPaintAdapter, FMaterialUpdateContext& MaterialUpdateContext);
+	double CalculateTargetEdgeLength(int TargetTriCount);
+	void StartPaintingTexture(UMeshComponent* InMeshComponent, const IMeshPaintComponentAdapter& GeometryInfo);
+	void UpdateResult();
+	void GatherTextureTriangles(IMeshPaintComponentAdapter* Adapter, int32 TriangleIndex, const int32 VertexIndices[3], TArray<FTexturePaintTriangleInfo>* TriangleInfo, TArray<FTexturePaintMeshSectionInfo>* SectionInfos, int32 UVChannelIndex);
+	bool Paint(const FVector& InRayOrigin, const FVector& InRayDirection);
+	bool Paint(const TArrayView<TPair<FVector, FVector>>& Rays);
+	void PaintTexture(FMeshPaintParameters& InParams, int32 UVChannel, TArray<FTexturePaintTriangleInfo>& InInfluencedTriangles, UMeshComponent* MeshComponent, const IMeshPaintComponentAdapter& GeometryInfo, FMeshPaintParameters* LastParams = nullptr);
+	bool PaintInternal(const TArrayView<TPair<FVector, FVector>>& Rays, EMeshPaintModeAction PaintAction, float PaintStrength);
+	void FinishPaintingTexture();
+	void FinishPainting();
 
 protected:
 	/** Textures eligible for painting retrieved from the current selection */
