Add defensive check for null material when removing mesh pain texture override.
The material is stored in a weak object pointer which can be null at this point if the map is being unloaded.
#rnx

diff --git a/Engine/Plugins/MeshPainting/Source/MeshPaintingToolset/Private/IMeshPaintComponentAdapter.cpp b/Engine/Plugins/MeshPainting/Source/MeshPaintingToolset/Private/IMeshPaintComponentAdapter.cpp
--- a/Engine/Plugins/MeshPainting/Source/MeshPaintingToolset/Private/IMeshPaintComponentAdapter.cpp
+++ b/Engine/Plugins/MeshPainting/Source/MeshPaintingToolset/Private/IMeshPaintComponentAdapter.cpp
@@ -283,11 +283,15 @@ namespace UE::MeshPaintingToolset
 							OverrideData->Count -= 1;
 							if (OverrideData->Count == 0)
 							{
-								UMaterialInterface* Material = Override.Key.Get();
-								const UTexture* SourceTexture = Override.Value.Get();
-								for (const ERHIFeatureLevel::Type LevelToUpdate : OverrideData->OverridenFeatureLevels)
+								if (UMaterialInterface* Material = Override.Key.Get())
 								{
-									Material->OverrideTexture(SourceTexture, nullptr, LevelToUpdate);
+									if (const UTexture* SourceTexture = Override.Value.Get())
+									{
+										for (const ERHIFeatureLevel::Type LevelToUpdate : OverrideData->OverridenFeatureLevels)
+										{
+											Material->OverrideTexture(SourceTexture, nullptr, LevelToUpdate);
+										}
+									}
 								}
 
 								OverridesData.RemoveByHash(OverrideHash, Override);
