Add world channel painting to mesh paint tool.
A new UMeshPaintChannelSet Data Asset defines how a set of named channels map from world space painting to texture/channel/UVs.
When mesh texture asset painting, the asset can be selected, and a channel chosen.
This then overrides the default mesh UV painting, and uses the world layer mappings to paint instead.
This is a first very quick iteration to get a round of user feedback.
Some expected follow up work would be: support paint color visualisation mode and maybe change UI to split functionality into its own mesh paint tab.

diff --git a/Engine/Plugins/MeshPainting/Source/MeshPaintEditorMode/MeshPaintEditorMode.Build.cs b/Engine/Plugins/MeshPainting/Source/MeshPaintEditorMode/MeshPaintEditorMode.Build.cs
--- a/Engine/Plugins/MeshPainting/Source/MeshPaintEditorMode/MeshPaintEditorMode.Build.cs
+++ b/Engine/Plugins/MeshPainting/Source/MeshPaintEditorMode/MeshPaintEditorMode.Build.cs
@@ -39,7 +39,8 @@ namespace UnrealBuildTool.Rules
 					"DynamicMesh",
 					"GeometryCore",
 					"ImageCore",
-					"MeshConversionEngineTypes"
+					"MeshConversionEngineTypes",
+					"Landscape"
 				}
             );
 
diff --git a/Engine/Plugins/MeshPainting/Source/MeshPaintEditorMode/Private/BrushSettingsCustomization.cpp b/Engine/Plugins/MeshPainting/Source/MeshPaintEditorMode/Private/BrushSettingsCustomization.cpp
--- a/Engine/Plugins/MeshPainting/Source/MeshPaintEditorMode/Private/BrushSettingsCustomization.cpp
+++ b/Engine/Plugins/MeshPainting/Source/MeshPaintEditorMode/Private/BrushSettingsCustomization.cpp
@@ -423,7 +423,7 @@ void FTexturePaintingSettingsCustomization::CustomizeDetails(IDetailLayoutBuilde
 	IDetailCategoryBuilder& BrushCategory = DetailLayout.EditCategory(TEXT("Brush"));
 	BrushCategory.SetSortOrder(2);
 	IDetailCategoryBuilder& TextureCategory = DetailLayout.EditCategory(TEXT("TexturePainting"));
-	TextureCategory.SetSortOrder(3);
+	TextureCategory.SetSortOrder(4);
 
 	TSharedRef<IPropertyHandle> RedChannel = DetailLayout.GetProperty("bWriteRed", UMeshTexturePaintingToolProperties::StaticClass());
 	RedChannel->MarkHiddenByCustomization();
@@ -482,13 +482,13 @@ void FTextureAssetPaintingSettingsCustomization::CustomizeDetails(IDetailLayoutB
 {
 	FTexturePaintingSettingsCustomization::CustomizeDetails(DetailLayout);
 
+	IDetailCategoryBuilder& WorldChannelCategory = DetailLayout.EditCategory(TEXT("WorldChannels"));
+	WorldChannelCategory.SetSortOrder(3);
 	IDetailCategoryBuilder& TextureCategory = DetailLayout.EditCategory(TEXT("TexturePainting"));
-	TextureCategory.SetSortOrder(3);
+	TextureCategory.SetSortOrder(4);
 
 	TSharedRef<IPropertyHandle> UVChannel = DetailLayout.GetProperty(GET_MEMBER_NAME_CHECKED(UMeshTextureAssetPaintingToolProperties, UVChannel));
-	UVChannel->MarkHiddenByCustomization();
-
-	TextureCategory.AddCustomRow(LOCTEXT("TexturePaintingUVLabel", "Texture Painting UV Channel"))
+	DetailLayout.EditDefaultProperty(UVChannel)->CustomWidget()
 	.NameContent()
 	[
 		UVChannel->CreatePropertyNameWidget()
diff --git a/Engine/Plugins/MeshPainting/Source/MeshPaintEditorMode/Private/MeshPaintMode.cpp b/Engine/Plugins/MeshPainting/Source/MeshPaintEditorMode/Private/MeshPaintMode.cpp
--- a/Engine/Plugins/MeshPainting/Source/MeshPaintEditorMode/Private/MeshPaintMode.cpp
+++ b/Engine/Plugins/MeshPainting/Source/MeshPaintEditorMode/Private/MeshPaintMode.cpp
@@ -27,6 +27,8 @@
 #include "LevelEditor.h"
 #include "TexturePaintToolset.h"
 #include "ToolContextInterfaces.h"
+#include "LandscapeNaniteComponent.h"
+#include "LandscapeProxy.h"
 
 #include UE_INLINE_GENERATED_CPP_BY_NAME(MeshPaintMode)
 
@@ -913,6 +915,24 @@ TArray<ComponentClass*> UMeshPaintMode::GetSelectedComponents() const
 			AActor* SelectedActor = Cast<AActor>(SelectionState.SelectedActors[SelectionIndex]);
 			if (SelectedActor)
 			{
+				// Code inside branch is only useful and compiles when ULandscapeNaniteComponent is derived from ComponentClass.
+				if constexpr (TIsDerivedFrom<ULandscapeNaniteComponent, ComponentClass>::Value)
+				{
+					// Selected landscape actors may contain many mesh components that we are not interested in (landscape grass, spline etc).
+					// To avoid performance/selection issues we assume that we only want to paint on the Nanite Landscape Components.
+					if (SelectedActor->IsA<ALandscapeProxy>())
+					{
+						TInlineComponentArray<ULandscapeNaniteComponent*> ActorComponents;
+						SelectedActor->GetComponents(ActorComponents);
+						for (ULandscapeNaniteComponent* Component : ActorComponents)
+						{
+							Components.AddUnique(Component);
+						}
+						
+						continue;
+					}
+				}
+
 				TInlineComponentArray<ComponentClass*> ActorComponents;
 				SelectedActor->GetComponents(ActorComponents);
 				for (ComponentClass* Component : ActorComponents)
diff --git a/Engine/Plugins/MeshPainting/Source/MeshPaintingToolset/Private/MeshPaintChannelSet.cpp b/Engine/Plugins/MeshPainting/Source/MeshPaintingToolset/Private/MeshPaintChannelSet.cpp
new file mode 100644
--- /dev/null
+++ b/Engine/Plugins/MeshPainting/Source/MeshPaintingToolset/Private/MeshPaintChannelSet.cpp
@@ -0,0 +1,18 @@
+// Copyright Epic Games, Inc. All Rights Reserved.
+
+#include "MeshPaintChannelSet.h"
+
+#include UE_INLINE_GENERATED_CPP_BY_NAME(MeshPaintChannelSet)
+
+FMeshPaintChannelDesc const* UMeshPaintChannelSet::GetChannelDesc(FName InChannelName) const
+{
+	for (FMeshPaintChannelDesc const& Channel : Channels)
+	{
+		if (InChannelName == Channel.ChannelName)
+		{
+			return &Channel;
+		}
+	}
+
+	return nullptr;
+}
diff --git a/Engine/Plugins/MeshPainting/Source/MeshPaintingToolset/Private/MeshPaintChannelSet.h b/Engine/Plugins/MeshPainting/Source/MeshPaintingToolset/Private/MeshPaintChannelSet.h
new file mode 100644
--- /dev/null
+++ b/Engine/Plugins/MeshPainting/Source/MeshPaintingToolset/Private/MeshPaintChannelSet.h
@@ -0,0 +1,59 @@
+// Copyright Epic Games, Inc. All Rights Reserved.
+
+#pragma once
+
+#include "Engine/DataAsset.h"
+#include "MeshPaintChannelSet.generated.h"
+
+class UTexture2D;
+
+USTRUCT()
+struct FMeshPaintChannelDesc
+{
+	GENERATED_BODY()
+
+	/** Name of the channel that will appear in the paint tool. */
+	UPROPERTY(EditAnywhere, Category=MeshPaint)
+	FName ChannelName;
+	/** Texture that channel paints to. */
+	UPROPERTY(EditAnywhere, Category=MeshPaint)
+	TObjectPtr<UTexture2D> PaintTexture;
+	/** True if the channel writes to the red channel. */
+	UPROPERTY(EditAnywhere, Category=MeshPaint)
+	bool bWriteRed = false;
+	/** True if the channel writes to the green channel. */
+	UPROPERTY(EditAnywhere, Category=MeshPaint)
+	bool bWriteGreen = false;
+	/** True if the channel writes to the blue channel. */
+	UPROPERTY(EditAnywhere, Category=MeshPaint)
+	bool bWriteBlue = false;
+	/** True if the channel writes to the alpha channel. */
+	UPROPERTY(EditAnywhere, Category=MeshPaint)
+	bool bWriteAlpha = false;
+	/** World Units per UV unit. */
+	UPROPERTY(EditAnywhere, Category=MeshPaint)
+	FVector2D UVToWorldScale = FVector2D::One();
+	/** UV offset to apply after scaling from world into UV space. */
+	UPROPERTY(EditAnywhere, Category=MeshPaint)
+	FVector2D WorldToUVBias = FVector2D::Zero();
+	/** Minimum clamp to apply to UV after scale and bias. */
+	UPROPERTY(EditAnywhere, Category=MeshPaint)
+	FVector2D ClampUVMin = FVector2D::Zero();
+	/** Maximum clamp to apply to UV after scale and bias. */
+	UPROPERTY(EditAnywhere, Category=MeshPaint)
+	FVector2D ClampUVMax = FVector2D::One();
+};
+
+UCLASS(MinimalAPI)
+class UMeshPaintChannelSet : public UDataAsset
+{
+	GENERATED_BODY()
+
+public:
+	/** The array of Channel Descriptions. */
+	UPROPERTY(EditAnywhere, Category=MeshPaint)
+	TArray<FMeshPaintChannelDesc> Channels;
+
+	/** Return a channel with the matching ChannelName. If none is found return nullptr. */
+	FMeshPaintChannelDesc const* GetChannelDesc(FName InChannelName) const;
+};
\ No newline at end of file
diff --git a/Engine/Plugins/MeshPainting/Source/MeshPaintingToolset/Private/MeshTexturePaintingTool.cpp b/Engine/Plugins/MeshPainting/Source/MeshPaintingToolset/Private/MeshTexturePaintingTool.cpp
--- a/Engine/Plugins/MeshPainting/Source/MeshPaintingToolset/Private/MeshTexturePaintingTool.cpp
+++ b/Engine/Plugins/MeshPainting/Source/MeshPaintingToolset/Private/MeshTexturePaintingTool.cpp
@@ -17,6 +17,7 @@
 #include "InteractiveToolManager.h"
 #include "Materials/MaterialInterface.h"
 #include "MaterialShared.h"
+#include "MeshPaintChannelSet.h"
 #include "MeshPaintHelpers.h"
 #include "MeshVertexPaintingTool.h"
 #include "ObjectCacheContext.h"
@@ -102,8 +103,11 @@ void UMeshTexturePaintingTool::Setup()
 void UMeshTexturePaintingTool::Shutdown(EToolShutdownType ShutdownType)
 {
 	FinishPainting();
-	
-	ClearAllTextureOverrides();
+
+	{
+		FMaterialUpdateContext MaterialUpdateContext(FMaterialUpdateContext::EOptions::Default|FMaterialUpdateContext::EOptions::SkipRecacheStaticPermutations);
+		ClearAllTextureOverrides(MaterialUpdateContext);
+	}
 
 	PaintTargetData.Empty();
 
@@ -181,10 +185,7 @@ void UMeshTexturePaintingTool::OnTick(float DeltaTime)
 				BucketFillParams.PaintAction = EMeshPaintModeAction::Paint;
 				BucketFillParams.BrushColor = TextureProperties->PaintColor;
 				BucketFillParams.BrushStrength = BrushStrength;
-				BucketFillParams.bWriteRed = TextureProperties->bWriteRed;
-				BucketFillParams.bWriteGreen = TextureProperties->bWriteGreen;
-				BucketFillParams.bWriteBlue = TextureProperties->bWriteBlue;
-				BucketFillParams.bWriteAlpha = TextureProperties->bWriteAlpha;
+				TextureProperties->GetWriteMask(BucketFillParams.bWriteRed, BucketFillParams.bWriteGreen, BucketFillParams.bWriteBlue, BucketFillParams.bWriteAlpha);
 				BucketFillParams.bUseFillBucket = true;
 			}
 
@@ -198,7 +199,7 @@ void UMeshTexturePaintingTool::OnTick(float DeltaTime)
 				}
 
 				const int32 UVChannel = GetSelectedUVChannel(SelectedComponent);
-				if (UVChannel >= MeshAdapter->GetNumUVChannels())
+				if (UVChannel < 0 || UVChannel >= MeshAdapter->GetNumUVChannels())
 				{
 					continue;
 				}
@@ -237,6 +238,8 @@ void UMeshTexturePaintingTool::OnTick(float DeltaTime)
 					TempAction.Execute(MeshAdapter.Get(), i / 3, TriangleIndices);
 				}
 
+				ApplyWorldChannelUVs(TrianglePaintInfoArray);
+
 				// Painting textures
 				UTexture2D* SelectedPaintTexure = GetSelectedPaintTexture(SelectedComponent);
 				if (PaintingTexture2D != nullptr && PaintingTexture2D != SelectedPaintTexure)
@@ -258,12 +261,14 @@ void UMeshTexturePaintingTool::OnTick(float DeltaTime)
 		UMeshComponent* FirstSelectedComponent = SelectedMeshComponents.IsValidIndex(0) ? SelectedMeshComponents[0] : nullptr;
 		if (MeshPaintingSubsystem->bNeedsRecache || (PaintableTextures.Num() > 0 && GetSelectedPaintTexture(FirstSelectedComponent) == nullptr))
 		{
-			ClearAllTextureOverrides();
+			FMaterialUpdateContext MaterialUpdateContext(FMaterialUpdateContext::EOptions::Default|FMaterialUpdateContext::EOptions::SkipRecacheStaticPermutations);
+
+			ClearAllTextureOverrides(MaterialUpdateContext);
 
 			CacheSelectionData();
 			CacheTexturePaintData();
 
-			SetAllTextureOverrides();
+			SetAllTextureOverrides(MaterialUpdateContext);
 		}
 	}
 
@@ -458,11 +463,8 @@ bool UMeshTexturePaintingTool::PaintInternal(const TArrayView<TPair<FVector, FVe
 					Params.InnerBrushDepth = BrushDepth - Params.BrushDepthFalloffRange;
 					Params.BrushStrength = BrushStrength;
 					Params.BrushToWorldMatrix = FMatrix(BrushXAxis, BrushYAxis, Params.BrushNormal, Params.BrushPosition);
-					Params.InverseBrushToWorldMatrix = Params.BrushToWorldMatrix.InverseFast();
-					Params.bWriteRed = TextureProperties->bWriteRed;
-					Params.bWriteGreen = TextureProperties->bWriteGreen;
-					Params.bWriteBlue = TextureProperties->bWriteBlue;
-					Params.bWriteAlpha = TextureProperties->bWriteAlpha;
+					Params.InverseBrushToWorldMatrix = Params.BrushToWorldMatrix.Inverse();
+					TextureProperties->GetWriteMask(Params.bWriteRed, Params.bWriteGreen, Params.bWriteBlue, Params.bWriteAlpha);
 					FVector BrushSpaceVertexPosition = Params.InverseBrushToWorldMatrix.TransformVector(FVector4(Params.BrushPosition, 1.0f));
 					Params.BrushPosition2D = FVector2f(BrushSpaceVertexPosition.X, BrushSpaceVertexPosition.Y);
 				}
@@ -489,7 +491,7 @@ bool UMeshTexturePaintingTool::PaintInternal(const TArrayView<TPair<FVector, FVe
 				}
 				
 				const int32 UVChannel = GetSelectedUVChannel(HoveredComponent);
-				if (UVChannel >= MeshAdapter->GetNumUVChannels())
+				if (UVChannel < 0 || UVChannel >= MeshAdapter->GetNumUVChannels())
 				{
 					continue;
 				}
@@ -528,6 +530,8 @@ bool UMeshTexturePaintingTool::PaintInternal(const TArrayView<TPair<FVector, FVe
 					continue;
 				}
 
+				ApplyWorldChannelUVs(TrianglePaintInfoArray);
+
 				// Painting textures
 				bAnyPaintApplied = true;
 
@@ -1292,12 +1296,10 @@ void UMeshTexturePaintingTool::FinishPaintingTexture()
 	PaintingTexture2D = nullptr;
 }
 
-void UMeshTexturePaintingTool::ClearAllTextureOverrides()
+void UMeshTexturePaintingTool::ClearAllTextureOverrides(FMaterialUpdateContext& MaterialUpdateContext)
 {
 	if (UMeshPaintingSubsystem* MeshPaintingSubsystem = GEngine->GetEngineSubsystem<UMeshPaintingSubsystem>())
 	{
-		FMaterialUpdateContext MaterialUpdateContext;
-
 		/** Remove all texture overrides which are currently stored and active */
 		for (decltype(PaintTargetData)::TIterator It(PaintTargetData); It; ++It)
 		{
@@ -1316,7 +1318,7 @@ void UMeshTexturePaintingTool::ClearAllTextureOverrides()
 	}
 }
 
-void UMeshTexturePaintingTool::SetAllTextureOverrides()
+void UMeshTexturePaintingTool::SetAllTextureOverrides(FMaterialUpdateContext& MaterialUpdateContext)
 {
 	UMeshPaintingSubsystem* MeshPaintingSubsystem = GEngine->GetEngineSubsystem<UMeshPaintingSubsystem>();
 	TArray<UMeshComponent*> SelectedMeshComponents = MeshPaintingSubsystem->GetSelectedMeshComponents();
@@ -1396,7 +1398,6 @@ void UMeshTexturePaintingTool::SetAllTextureOverrides()
 		}
 
 		// Apply the overrides.
-		FMaterialUpdateContext MaterialUpdateContext;
 		for (UMeshComponent* MeshComponent : PaintableMeshComponents)
 		{
 			TSharedPtr<IMeshPaintComponentAdapter> MeshAdapter = MeshPaintingSubsystem->GetAdapterForComponent(MeshComponent);
@@ -1539,7 +1540,36 @@ void UMeshTextureAssetPaintingTool::OnPropertyModified(UObject* PropertySet, FPr
 {
 	Super::OnPropertyModified(PropertySet, Property);
 
-	if (Property && (Property->GetFName() == GET_MEMBER_NAME_CHECKED(UMeshTextureAssetPaintingToolProperties, PaintTexture)))
+	// Need to update the currently painted texture.
+	bool bUpdateTextureOverrides = false;
+	// Need to recache PaintableTextures.
+	bool bUpdateCachedTexturePaintData = false;
+	
+	if (Property && Property->GetFName() == GET_MEMBER_NAME_CHECKED(UMeshTextureAssetPaintingToolProperties, ChannelSet))
+	{
+		// Changing channel set changes the PaintableTextures.
+		bUpdateCachedTexturePaintData = true;
+	}
+
+	if (Property &&	Property->GetFName() == GET_MEMBER_NAME_CHECKED(UMeshTextureAssetPaintingToolProperties, ChannelName))
+	{
+		UTexture2D* ChannelPaintTexture = nullptr;
+		if (AssetProperties->ChannelSet)
+		{
+			if (FMeshPaintChannelDesc const* Channel = AssetProperties->ChannelSet->GetChannelDesc(AssetProperties->ChannelName))
+			{
+				ChannelPaintTexture = Channel->PaintTexture;
+			}
+		}
+		
+		bUpdateTextureOverrides |= AssetProperties->PaintTexture != ChannelPaintTexture || AssetProperties->UVChannel != 0;
+		bUpdateCachedTexturePaintData |= PaintableTextures.Num() != 1 || PaintableTextures[0].Texture != ChannelPaintTexture;
+
+		AssetProperties->PaintTexture = ChannelPaintTexture;
+		AssetProperties->UVChannel = 0;
+	}
+
+	if (Property && Property->GetFName() == GET_MEMBER_NAME_CHECKED(UMeshTextureAssetPaintingToolProperties, PaintTexture))
 	{
 		// Find the selected texture and apply it's UV channel.
 		for (FPaintableTexture& PaintableTexture : PaintableTextures)
@@ -1552,13 +1582,43 @@ void UMeshTextureAssetPaintingTool::OnPropertyModified(UObject* PropertySet, FPr
 		}
 		
 		// Need to recreate the render target overrides with the newly selected texture.
-		ClearAllTextureOverrides();
-		SetAllTextureOverrides();
+		bUpdateTextureOverrides = true;
+	}
+	
+	if (bUpdateTextureOverrides || bUpdateCachedTexturePaintData)
+	{
+		FMaterialUpdateContext MaterialUpdateContext(FMaterialUpdateContext::EOptions::Default|FMaterialUpdateContext::EOptions::SkipRecacheStaticPermutations);
+
+		ClearAllTextureOverrides(MaterialUpdateContext);
+		if (bUpdateCachedTexturePaintData)
+		{
+			CacheTexturePaintData();
+		}
+		SetAllTextureOverrides(MaterialUpdateContext);
 	}
 }
 
 void UMeshTextureAssetPaintingTool::CacheTexturePaintData()
 {
+	if (AssetProperties->ChannelSet != nullptr)
+	{
+		PaintableTextures.Reset();
+		AssetProperties->PaintTexture = nullptr;
+		AssetProperties->UVChannel = INDEX_NONE;
+
+		if (FMeshPaintChannelDesc const* Channel = AssetProperties->ChannelSet->GetChannelDesc(AssetProperties->ChannelName))
+		{
+			if (Channel->PaintTexture)
+			{
+				PaintableTextures.Add(Channel->PaintTexture);
+				AssetProperties->PaintTexture = Channel->PaintTexture;
+				AssetProperties->UVChannel = 0;
+			}
+		}
+
+		return;
+	}
+
 	UMeshPaintingSubsystem* MeshPaintingSubsystem = GEngine->GetEngineSubsystem<UMeshPaintingSubsystem>();
 	if (MeshPaintingSubsystem)
 	{
@@ -1660,6 +1720,10 @@ UTexture* UMeshTextureAssetPaintingTool::GetSelectedPaintTextureWithOverride() c
 
 void UMeshTextureAssetPaintingTool::CycleTextures(int32 Direction)
 {
+	if (AssetProperties->ChannelSet != nullptr)
+	{
+		return;
+	}
 	if (!PaintableTextures.Num())
 	{
 		return;
@@ -1682,4 +1746,76 @@ void UMeshTextureAssetPaintingTool::CycleTextures(int32 Direction)
 	}
 }
 
+void UMeshTextureAssetPaintingTool::ApplyWorldChannelUVs(TArrayView<FTexturePaintTriangleInfo> TriangleInfo) const
+{
+	if (AssetProperties->ChannelSet)
+	{
+		if (FMeshPaintChannelDesc const* Channel = AssetProperties->ChannelSet->GetChannelDesc(AssetProperties->ChannelName))
+		{
+			// Scale is stored as UV->World to make UI easier to use (it's difficult to enter very small numbers with precision).
+			const float WorldToUVScaleX = FMath::IsNearlyZero(Channel->UVToWorldScale.X) ? 1.f : 1.f / Channel->UVToWorldScale.X;
+			const float WorldToUVScaleY = FMath::IsNearlyZero(Channel->UVToWorldScale.Y) ? 1.f : 1.f / Channel->UVToWorldScale.Y;
+
+			for (FTexturePaintTriangleInfo& Info : TriangleInfo)
+			{
+				Info.TriUVs[0].X = FMath::Clamp(Info.TriVertices[0].X * WorldToUVScaleX + Channel->WorldToUVBias.X, Channel->ClampUVMin.X, Channel->ClampUVMax.X);
+				Info.TriUVs[0].Y = FMath::Clamp(Info.TriVertices[0].Y * WorldToUVScaleY + Channel->WorldToUVBias.Y, Channel->ClampUVMin.Y, Channel->ClampUVMax.Y);
+				Info.TriUVs[1].X = FMath::Clamp(Info.TriVertices[1].X * WorldToUVScaleX + Channel->WorldToUVBias.X, Channel->ClampUVMin.X, Channel->ClampUVMax.X);
+				Info.TriUVs[1].Y = FMath::Clamp(Info.TriVertices[1].Y * WorldToUVScaleY + Channel->WorldToUVBias.Y, Channel->ClampUVMin.Y, Channel->ClampUVMax.Y);
+				Info.TriUVs[2].X = FMath::Clamp(Info.TriVertices[2].X * WorldToUVScaleX + Channel->WorldToUVBias.X, Channel->ClampUVMin.X, Channel->ClampUVMax.X);
+				Info.TriUVs[2].Y = FMath::Clamp(Info.TriVertices[2].Y * WorldToUVScaleY + Channel->WorldToUVBias.Y, Channel->ClampUVMin.Y, Channel->ClampUVMax.Y);
+			}
+		}
+	}
+}
+
+TArray<FName> UMeshTextureAssetPaintingToolProperties::GetChannelNames() const
+{
+	TArray<FName> ChannelNames;
+	if (ChannelSet != nullptr)
+	{
+		for (FMeshPaintChannelDesc const& Channel : ChannelSet->Channels)
+		{
+			ChannelNames.AddUnique(Channel.ChannelName);
+		}
+	}
+	return ChannelNames;
+}
+
+bool UMeshTextureAssetPaintingToolProperties::CanEditChange(const FProperty* InProperty) const
+{
+	bool bCanEdit = Super::CanEditChange(InProperty);
+	if (InProperty != nullptr && 
+		(InProperty->GetFName() == GET_MEMBER_NAME_CHECKED(UMeshTexturePaintingToolProperties, bWriteRed) ||
+		InProperty->GetFName() == GET_MEMBER_NAME_CHECKED(UMeshTexturePaintingToolProperties, bWriteGreen) ||
+		InProperty->GetFName() == GET_MEMBER_NAME_CHECKED(UMeshTexturePaintingToolProperties, bWriteBlue) ||
+		InProperty->GetFName() == GET_MEMBER_NAME_CHECKED(UMeshTexturePaintingToolProperties, bWriteAlpha) ||
+		InProperty->GetFName() == GET_MEMBER_NAME_CHECKED(UMeshTexturePaintingToolProperties, bEnableSeamPainting)))
+	{
+		bCanEdit &= ChannelSet == nullptr;
+	}
+	return bCanEdit;
+}
+
+void UMeshTextureAssetPaintingToolProperties::GetWriteMask(bool& bOutWriteRed, bool& bOutWriteGreen, bool& bOutWriteBlue, bool& bOutWriteAlpha) const
+{
+	if (ChannelSet == nullptr)
+	{
+		Super::GetWriteMask(bOutWriteRed, bOutWriteGreen, bOutWriteBlue, bOutWriteAlpha);
+		return;
+	}
+
+	if (FMeshPaintChannelDesc const* Channel = ChannelSet->GetChannelDesc(ChannelName))
+	{
+		bOutWriteRed = Channel->bWriteRed;
+		bOutWriteGreen = Channel->bWriteGreen;
+		bOutWriteBlue = Channel->bWriteBlue;
+		bOutWriteAlpha = Channel->bWriteAlpha;
+		return;
+	}
+		
+	// Didn't find channel so disable all writes.
+	bOutWriteRed = bOutWriteGreen = bOutWriteBlue = bOutWriteAlpha = false;
+}
+
 #undef LOCTEXT_NAMESPACE
diff --git a/Engine/Plugins/MeshPainting/Source/MeshPaintingToolset/Public/MeshTexturePaintingTool.h b/Engine/Plugins/MeshPainting/Source/MeshPaintingToolset/Public/MeshTexturePaintingTool.h
--- a/Engine/Plugins/MeshPainting/Source/MeshPaintingToolset/Public/MeshTexturePaintingTool.h
+++ b/Engine/Plugins/MeshPainting/Source/MeshPaintingToolset/Public/MeshTexturePaintingTool.h
@@ -15,6 +15,7 @@ enum class EToolShutdownType : uint8;
 class FMaterialUpdateContext;
 class FScopedTransaction;
 class IMeshPaintComponentAdapter;
+class UMeshPaintChannelSet;
 class UMeshToolManager;
 class UTexture2D;
 struct FPaintRayResults;
@@ -94,6 +95,15 @@ public:
 	/** Whether or not to apply Texture Color Painting to the Alpha Channel */
 	UPROPERTY(EditAnywhere, Category = ColorPainting, meta = (DisplayName = "Alpha"))
 	bool bWriteAlpha = false;
+
+	/** Get all of the write mask properties. */
+	virtual void GetWriteMask(bool& bOutWriteRed, bool& bOutWriteGreen, bool& bOutWriteBlue, bool& bOutWriteAlpha) const
+	{
+		bOutWriteRed = bWriteRed;
+		bOutWriteGreen = bWriteGreen;
+		bOutWriteBlue = bWriteBlue;
+		bOutWriteAlpha = bWriteAlpha;
+	}
 };
 
 /**
@@ -119,13 +129,36 @@ class UMeshTextureAssetPaintingToolProperties : public UMeshTexturePaintingToolP
 	GENERATED_BODY()
 
 public:
+	/** The MeshPaintChannelSet asset that defines how a group of chanels map onto texture assets. */
+	UPROPERTY(EditAnywhere, Category = WorldChannels, meta = (DisplayThumbnail = "false"))
+	TObjectPtr<UMeshPaintChannelSet> ChannelSet;
+
+	/** The currently selected Channel from the ChannelSet. */
+	UPROPERTY(EditAnywhere, Category = WorldChannels, meta = (GetOptions = "GetChannelNames", editcondition = "ChannelSet != nullptr"))
+	FName ChannelName;
+
 	/** UV channel which should be used for painting textures. */
-	UPROPERTY(EditAnywhere, Category = TexturePainting, meta = (TransientToolProperty))
+	UPROPERTY(EditAnywhere, Category = TexturePainting, meta = (TransientToolProperty, editcondition = "ChannelSet == nullptr"))
 	int32 UVChannel = 0;
 
 	/** Texture to which painting should be applied. */
-	UPROPERTY(EditAnywhere, Category = TexturePainting, meta = (DisplayThumbnail = "true", TransientToolProperty))
+	UPROPERTY(EditAnywhere, Category = TexturePainting, meta = (TransientToolProperty, DisplayThumbnail = "true", editcondition = "ChannelSet == nullptr"))
 	TObjectPtr<UTexture2D> PaintTexture;
+
+protected:
+	/** Helper function to get the ChannelNames from the currently selected ChannelSet. */
+	UFUNCTION()
+    TArray<FName> GetChannelNames() const;
+
+#if WITH_EDITOR
+	// Begin UObject Interface.
+	virtual bool CanEditChange(const FProperty* InProperty) const override;
+	// End UObject Interface.
+#endif
+
+	// Begin UMeshTexturePaintingToolProperties Interface.
+	virtual void GetWriteMask(bool& bOutWriteRed, bool& bOutWriteGreen, bool& bOutWriteBlue, bool& bOutWriteAlpha) const override;
+	// End UMeshTexturePaintingToolProperties Interface.
 };
 
 
@@ -170,12 +203,13 @@ protected:
 	UE_API FPaintTexture2DData* GetPaintTargetData(const UTexture2D* InTexture);
 	UE_API FPaintTexture2DData* AddPaintTargetData(UTexture2D* InTexture);
 	
-	UE_API void SetAllTextureOverrides();
-	UE_API void ClearAllTextureOverrides();
+	UE_API void SetAllTextureOverrides(FMaterialUpdateContext& MaterialUpdateContext);
+	UE_API void ClearAllTextureOverrides(FMaterialUpdateContext& MaterialUpdateContext);
 
 	virtual UTexture2D* GetSelectedPaintTexture(UMeshComponent const* InMeshComponent) const { return nullptr; }
 	virtual void CacheTexturePaintData() {}
 	virtual bool CanPaintTextureToComponent(UTexture* InTexture, UMeshComponent const* InMeshComponent) const { return false; }
+	virtual void ApplyWorldChannelUVs(TArrayView<FTexturePaintTriangleInfo> TriangleInfo) const {}
 
 private:
 	void CacheSelectionData();
@@ -306,6 +340,7 @@ protected:
 	UE_API virtual void GetModifiedTexturesToSave(TArray<UObject*>& OutTexturesToSave) const override;
 	UE_API virtual void CacheTexturePaintData() override;
 	UE_API virtual bool CanPaintTextureToComponent(UTexture* InTexture, UMeshComponent const* InMeshComponent) const override;
+	UE_API virtual void ApplyWorldChannelUVs(TArrayView<FTexturePaintTriangleInfo> TriangleInfo) const override;
 	// End UMeshTexturePaintingTool Interface.
 
 private:
