Fix for PVSStudio divide by zero warnings.
Divide by zero shouldn't ever happen because if InParams.TotalWeightCount==1 then InParams.PaintWeightIndex==0 and we branch away from the division.
Also calling code only calls with InParams.TotalWeightCount in range [2,5] from EMeshPaintWeightTypes.
Modify the check() at the top of the function to make this explicit.
#rnx

diff --git a/Engine/Plugins/MeshPainting/Source/MeshPaintingToolset/Private/MeshPaintHelpers.cpp b/Engine/Plugins/MeshPainting/Source/MeshPaintingToolset/Private/MeshPaintHelpers.cpp
--- a/Engine/Plugins/MeshPainting/Source/MeshPaintingToolset/Private/MeshPaintHelpers.cpp
+++ b/Engine/Plugins/MeshPainting/Source/MeshPaintingToolset/Private/MeshPaintHelpers.cpp
@@ -288,7 +288,7 @@ void UMeshPaintingSubsystem::ApplyVertexColorPaint(const FMeshPaintParameters &I
 void UMeshPaintingSubsystem::ApplyVertexWeightPaint(const FMeshPaintParameters &InParams, const FLinearColor &OldColor, FLinearColor &NewColor, const float PaintAmount)
 {
 	// Total number of texture blend weights we're using
-	check(InParams.TotalWeightCount > 0);
+	check(InParams.TotalWeightCount > 1);
 	check(InParams.TotalWeightCount <= MeshPaintDefs::MaxSupportedWeights);
 
 	// True if we should assume the last weight index is composed of one minus the sum of all
diff --git a/Engine/Source/Editor/MeshPaint/Private/MeshPaintHelpers.cpp b/Engine/Source/Editor/MeshPaint/Private/MeshPaintHelpers.cpp
--- a/Engine/Source/Editor/MeshPaint/Private/MeshPaintHelpers.cpp
+++ b/Engine/Source/Editor/MeshPaint/Private/MeshPaintHelpers.cpp
@@ -247,7 +247,7 @@ void MeshPaintHelpers::ApplyVertexColorPaint(const FMeshPaintParameters &InParam
 void MeshPaintHelpers::ApplyVertexWeightPaint(const FMeshPaintParameters &InParams, const FLinearColor &OldColor, const float PaintAmount, FLinearColor &NewColor)
 {
 	// Total number of texture blend weights we're using
-	check(InParams.TotalWeightCount > 0);
+	check(InParams.TotalWeightCount > 1);
 	check(InParams.TotalWeightCount <= MeshPaintDefs::MaxSupportedWeights);
 
 	// True if we should assume the last weight index is composed of one minus the sum of all
