MeshPaint: Reset selection after mesh compilation to take into account nanite mesh changes.
Also display selection warnings for setups where painting isn't possible and the reason may not be obvious why to the user.
#rnx
#FailedValidations Preflight

diff --git a/Engine/Plugins/MeshPainting/Source/MeshPaintEditorMode/Private/BrushSettingsCustomization.cpp b/Engine/Plugins/MeshPainting/Source/MeshPaintEditorMode/Private/BrushSettingsCustomization.cpp
--- a/Engine/Plugins/MeshPainting/Source/MeshPaintEditorMode/Private/BrushSettingsCustomization.cpp
+++ b/Engine/Plugins/MeshPainting/Source/MeshPaintEditorMode/Private/BrushSettingsCustomization.cpp
@@ -50,6 +50,47 @@ void FMeshPaintModeSettingsCustomization::CustomizeDetails(IDetailLayoutBuilder&
 	IDetailCategoryBuilder& ResourceCategory = DetailLayout.EditCategory(TEXT("ResourceUsage"));
 	ResourceCategory.SetSortOrder(0);
 
+	ResourceCategory.AddCustomRow(LOCTEXT("SelectionStateRow", "Selection State"))
+		.NameContent()
+		[
+			SNew(STextBlock)
+				.Font(IDetailLayoutBuilder::GetDetailFont())
+				.Text(LOCTEXT("SelectedComponentsName", "Selected Components"))
+		]
+		.ValueContent()
+		[
+			SNew(SHorizontalBox)
+				+ SHorizontalBox::Slot()
+				[
+					SNew(STextBlock)
+						.Font(IDetailLayoutBuilder::GetDetailFont())
+						.Text_Lambda([]() -> FText
+						{
+							const int32 NumSelectedComponents = UMeshPaintMode::GetMeshPaintMode()->GetSelectedComponents<UMeshComponent>().Num();
+							return FText::Format(LOCTEXT("SelectionCount", "{0}"), FText::AsNumber(NumSelectedComponents));
+						})
+				]
+				+ SHorizontalBox::Slot()
+				.HAlign(HAlign_Left)
+				.VAlign(VAlign_Center)
+				[
+					SNew(SImage)
+						.Image(FCoreStyle::Get().GetBrush("Icons.Warning"))
+						.ColorAndOpacity(FLinearColor::Yellow)
+						.Visibility_Lambda([this]()
+						{
+							FText Message;
+							return UMeshPaintMode::GetMeshPaintMode()->GetSelectionWarning(Message) ? EVisibility::Visible : EVisibility::Hidden;
+						})
+						.ToolTipText_Lambda([this]()
+						{
+							FText Message;
+							UMeshPaintMode::GetMeshPaintMode()->GetSelectionWarning(Message);
+							return Message;
+						})
+				]
+		];
+
 	ResourceCategory.AddCustomRow(NSLOCTEXT("VertexPaintSettings", "VertexColorSizeRow", "Vertex Color Size"))
 		.NameContent()
 		[
@@ -532,6 +573,7 @@ void FTextureAssetPaintingSettingsCustomization::CustomizeDetails(IDetailLayoutB
 		[
 			SNew(SImage)
 			.Image(FCoreStyle::Get().GetBrush("Icons.Warning"))
+			.ColorAndOpacity(FLinearColor::Yellow)
 			.Visibility_Lambda([this]()
 			{
 				const int32 UVChannel = UMeshPaintMode::GetTextureAssetToolProperties() ? UMeshPaintMode::GetTextureAssetToolProperties()->UVChannel : 0;
diff --git a/Engine/Plugins/MeshPainting/Source/MeshPaintEditorMode/Private/MeshPaintMode.cpp b/Engine/Plugins/MeshPainting/Source/MeshPaintEditorMode/Private/MeshPaintMode.cpp
--- a/Engine/Plugins/MeshPainting/Source/MeshPaintEditorMode/Private/MeshPaintMode.cpp
+++ b/Engine/Plugins/MeshPainting/Source/MeshPaintEditorMode/Private/MeshPaintMode.cpp
@@ -1,6 +1,8 @@
 // Copyright Epic Games, Inc. All Rights Reserved.
 
 #include "MeshPaintMode.h"
+
+#include "AssetCompilingManager.h"
 #include "EditorModeManager.h"
 #include "MeshPaintModeCommands.h"
 #include "EditorViewportClient.h"
@@ -27,6 +29,7 @@
 #include "LevelEditor.h"
 #include "TexturePaintToolset.h"
 #include "ToolContextInterfaces.h"
+#include "LandscapeComponent.h"
 #include "LandscapeNaniteComponent.h"
 #include "LandscapeProxy.h"
 
@@ -181,6 +184,8 @@ void UMeshPaintMode::Enter()
 	FLevelEditorModule& LevelEditor = FModuleManager::GetModuleChecked<FLevelEditorModule>(FName(TEXT("LevelEditor")));
 	LevelEditor.OnRedrawLevelEditingViewports().AddUObject(this, &UMeshPaintMode::UpdateOnMaterialChange);
 
+	FAssetCompilingManager::Get().OnAssetPostCompileEvent().AddUObject(this, &UMeshPaintMode::UpdateOnPostCompile);
+
 	// some global cvars can affect whether painting is valid (nanite on/off etc)
 	CVarDelegateHandle = IConsoleManager::Get().RegisterConsoleVariableSink_Handle(FConsoleCommandDelegate::CreateLambda([this]{ bRecacheValidForPaint = true; }));
 }
@@ -211,6 +216,8 @@ void UMeshPaintMode::Exit()
 	FLevelEditorModule& LevelEditor = FModuleManager::GetModuleChecked<FLevelEditorModule>(FName(TEXT("LevelEditor")));
 	LevelEditor.OnRedrawLevelEditingViewports().RemoveAll(this);
 
+	FAssetCompilingManager::Get().OnAssetPostCompileEvent().RemoveAll(this);
+
 	IConsoleManager::Get().UnregisterConsoleVariableSink_Handle(CVarDelegateHandle);
 	CVarDelegateHandle = {};
 }
@@ -545,6 +552,12 @@ void UMeshPaintMode::UpdateOnMaterialChange(bool bInvalidateHitProxies)
 	bRecacheValidForPaint = true;
 }
 
+void UMeshPaintMode::UpdateOnPostCompile(const TArray<FAssetCompileData>& Data)
+{
+	// Can happen after add/remove nanite data.
+	bRecacheValidForPaint = true;
+}
+
 void UMeshPaintMode::OnObjectsReplaced(const TMap<UObject*, UObject*>& InOldToNewInstanceMap)
 {
 	if (UMeshPaintingSubsystem* MeshPaintingSubsystem = GEngine->GetEngineSubsystem<UMeshPaintingSubsystem>())
@@ -910,9 +923,8 @@ TArray<ComponentClass*> UMeshPaintMode::GetSelectedComponents() const
 
 	if (Components.Num() == 0)
 	{
-		for (int32 SelectionIndex = 0; SelectionIndex < SelectionState.SelectedActors.Num(); ++SelectionIndex)
+		for (AActor* SelectedActor : SelectionState.SelectedActors)
 		{
-			AActor* SelectedActor = Cast<AActor>(SelectionState.SelectedActors[SelectionIndex]);
 			if (SelectedActor)
 			{
 				// Code inside branch is only useful and compiles when ULandscapeNaniteComponent is derived from ComponentClass.
@@ -951,7 +963,6 @@ template TArray<USkeletalMeshComponent*> UMeshPaintMode::GetSelectedComponents<U
 template TArray<UMeshComponent*> UMeshPaintMode::GetSelectedComponents<UMeshComponent>() const;
 template TArray<UGeometryCollectionComponent*> UMeshPaintMode::GetSelectedComponents<UGeometryCollectionComponent>() const;
 
-
 void UMeshPaintMode::UpdateCachedDataSizes()
 {
 	CachedVertexDataSize = 0;
@@ -1434,4 +1445,49 @@ bool UMeshPaintMode::CanCycleMeshLODs() const
 	return PaletteName == UMeshPaintMode::MeshPaintMode_VertexColor || PaletteName == UMeshPaintMode::MeshPaintMode_VertexWeights;
 }
 
+bool UMeshPaintMode::GetSelectionWarning(FText& OutMessage) const
+{
+	FToolBuilderState SelectionState;
+	GetToolManager()->GetContextQueriesAPI()->GetCurrentSelectionState(SelectionState);
+
+	FString ActiveTool = GetToolManager()->GetActiveToolName(EToolSide::Mouse);
+	if (ActiveTool == VertexSelectToolName || ActiveTool == VertexColorPaintToolName || ActiveTool == VertexWeightPaintToolName)
+	{
+		for (UStaticMeshComponent* StaticMeshComponent : GetSelectedComponents<UStaticMeshComponent>())
+		{
+			if (!StaticMeshComponent->CanMeshPaintVertexColors())
+			{
+				OutMessage = LOCTEXT("LevelMeshPainter_NaniteVertexPaintWarning", "Part of the selection cannot be vertex painted because it has disabled vertex paint or because it is currently rendered with nanite.");
+				return true;
+			}
+		}
+	}
+	if (ActiveTool == TextureColorSelectToolName || ActiveTool == TextureColorPaintToolName)
+	{
+		if (CanAddMeshPaintTextures())
+		{
+			OutMessage = LOCTEXT("LevelMeshPainter_AddMeshPaintTextureWarning", "Part of the selection has no Mesh Paint Texture and cannot be painted. Use the Add button to fix.");
+			return true;
+		}
+	}
+
+	for (AActor* SelectedActor : SelectionState.SelectedActors)
+	{
+		if (SelectedActor)
+		{
+			if (SelectedActor->IsA<ALandscapeProxy>())
+			{
+				TInlineComponentArray<ULandscapeNaniteComponent*> NaniteComponents;
+				SelectedActor->GetComponents<ULandscapeNaniteComponent>(NaniteComponents);
+				if (NaniteComponents.Num() == 0)
+				{
+					OutMessage = LOCTEXT("LevelMeshPainter_NaniteLandscapeWarning", "Part of the selected landscape has no nanite mesh and cannot be painted.");
+					return true;
+				}
+			}
+		}
+	}
+	return false;
+}
+
 #undef LOCTEXT_NAMESPACE
diff --git a/Engine/Plugins/MeshPainting/Source/MeshPaintEditorMode/Private/MeshPaintMode.h b/Engine/Plugins/MeshPainting/Source/MeshPaintEditorMode/Private/MeshPaintMode.h
--- a/Engine/Plugins/MeshPainting/Source/MeshPaintEditorMode/Private/MeshPaintMode.h
+++ b/Engine/Plugins/MeshPainting/Source/MeshPaintEditorMode/Private/MeshPaintMode.h
@@ -67,6 +67,9 @@ public:
 	/** Returns resource size of mesh paint textures for the currently selected components. */
 	uint32 GetMeshPaintTextureResourceSizeInBytes() const { return CachedMeshPaintTextureResourceSize; }
 
+	/** Returns true and fills the message if there is the current selection is not in an expectd state for the current painting mode. */
+	bool GetSelectionWarning(FText& OutMessage) const;
+
 protected:
 	/** Binds UI commands to actions for the mesh paint mode */
 	virtual void BindCommands() override;
@@ -84,6 +87,7 @@ protected:
 
 	void UpdateSelectedMeshes();
 	void UpdateOnMaterialChange(bool bInvalidateHitProxies);
+	void UpdateOnPostCompile(const TArray<struct FAssetCompileData>& Data);
 	void OnObjectsReplaced(const TMap<UObject*, UObject*>& InOldToNewInstanceMap);
 	void OnResetViewMode();
 	void OnVertexPaintFinished();
