Set UMeshPaintChannelSet as BlueprintType.
#rnx

diff --git a/Engine/Plugins/MeshPainting/Source/MeshPaintingToolset/Private/MeshPaintChannelSet.h b/Engine/Plugins/MeshPainting/Source/MeshPaintingToolset/Private/MeshPaintChannelSet.h
--- a/Engine/Plugins/MeshPainting/Source/MeshPaintingToolset/Private/MeshPaintChannelSet.h
+++ b/Engine/Plugins/MeshPainting/Source/MeshPaintingToolset/Private/MeshPaintChannelSet.h
@@ -7,7 +7,7 @@
 
 class UTexture2D;
 
-USTRUCT()
+USTRUCT(BlueprintType)
 struct FMeshPaintChannelDesc
 {
 	GENERATED_BODY()
@@ -44,7 +44,7 @@ struct FMeshPaintChannelDesc
 	FVector2D ClampUVMax = FVector2D::One();
 };
 
-UCLASS(MinimalAPI)
+UCLASS(BlueprintType, MinimalAPI)
 class UMeshPaintChannelSet : public UDataAsset
 {
 	GENERATED_BODY()
