Fix MeshPaint recreate of invalid MeshAdapter objects.
Validation was only checking the SharedPtr validity, but needs to call IsValid() on the underlying object too.
Failing to recreate led to the UI showing the associated mesh as unpaintable.
#rnx

diff --git a/Engine/Plugins/MeshPainting/Source/MeshPaintingToolset/Private/MeshPaintHelpers.cpp b/Engine/Plugins/MeshPainting/Source/MeshPaintingToolset/Private/MeshPaintHelpers.cpp
--- a/Engine/Plugins/MeshPainting/Source/MeshPaintingToolset/Private/MeshPaintHelpers.cpp
+++ b/Engine/Plugins/MeshPainting/Source/MeshPaintingToolset/Private/MeshPaintHelpers.cpp
@@ -1605,7 +1605,7 @@ void UMeshPaintingSubsystem::UpdatePaintSupportState()
 		if (UMeshComponent* MeshComponent = MeshComponentWeak.Get())
 		{
 			TSharedPtr<IMeshPaintComponentAdapter> MeshAdapter = GetAdapterForComponent(MeshComponent);
-			if (!MeshAdapter.IsValid())
+			if (!MeshAdapter.IsValid() || !MeshAdapter->IsValid())
 			{
 				MeshAdapter = FMeshPaintComponentAdapterFactory::CreateAdapterForMesh(MeshComponent, 0);
 				if (MeshAdapter)
diff --git a/Engine/Plugins/MeshPainting/Source/MeshPaintingToolset/Private/MeshPaintInteractions.cpp b/Engine/Plugins/MeshPainting/Source/MeshPaintingToolset/Private/MeshPaintInteractions.cpp
--- a/Engine/Plugins/MeshPainting/Source/MeshPaintingToolset/Private/MeshPaintInteractions.cpp
+++ b/Engine/Plugins/MeshPainting/Source/MeshPaintingToolset/Private/MeshPaintInteractions.cpp
@@ -100,7 +100,7 @@ bool UMeshPaintSelectionMechanic::FindClickedComponentsAndCacheAdapters(const FI
 					{
 						UMeshComponent* MeshComponent = Cast<UMeshComponent>(CandidateComponent);
 						TSharedPtr<IMeshPaintComponentAdapter> MeshAdapter = MeshPaintingSubsystem->GetAdapterForComponent(MeshComponent);
-						if (!MeshAdapter.IsValid())
+						if (!MeshAdapter.IsValid() || !MeshAdapter->IsValid())
 						{
 							MeshAdapter = FMeshPaintComponentAdapterFactory::CreateAdapterForMesh(MeshComponent, 0);
 							if (MeshAdapter)
