MeshPaint textures outer to the component and not the world package.
This fixes issue where texture editors for MeshPaint textures don't close on level unload in UEditorEngine::CloseEditedWorldAssets().
Not closing the editors leads to a fatal error about a GC leak from the UWorld.
This fix will only apply to newly created textures and not fix up existing ones.
#tests Cooked builds with WP and not

diff --git a/Engine/Plugins/MeshPainting/Source/MeshPaintEditorMode/Private/MeshPaintModeHelpers.cpp b/Engine/Plugins/MeshPainting/Source/MeshPaintEditorMode/Private/MeshPaintModeHelpers.cpp
--- a/Engine/Plugins/MeshPainting/Source/MeshPaintEditorMode/Private/MeshPaintModeHelpers.cpp
+++ b/Engine/Plugins/MeshPainting/Source/MeshPaintEditorMode/Private/MeshPaintModeHelpers.cpp
@@ -357,7 +357,7 @@ void UMeshPaintModeSubsystem::ImportMeshPaintTextureFromVertexColors(UMeshCompon
 	FImage ConvertedImage;
 	ResultImage.CopyTo(ConvertedImage, ERawImageFormat::BGRA8, EGammaSpace::sRGB);
 
-	UMeshPaintVirtualTexture* NewTexture = NewObject<UMeshPaintVirtualTexture>(StaticMeshComponent->GetOutermost());
+	UMeshPaintVirtualTexture* NewTexture = NewObject<UMeshPaintVirtualTexture>(StaticMeshComponent);
 	NewTexture->Source.Init(ConvertedImage);
 	NewTexture->OwningComponent = MakeWeakObjectPtr(StaticMeshComponent);
 	NewTexture->UpdateResource();
diff --git a/Engine/Plugins/MeshPainting/Source/MeshPaintingToolset/Private/MeshPaintHelpers.cpp b/Engine/Plugins/MeshPainting/Source/MeshPaintingToolset/Private/MeshPaintHelpers.cpp
--- a/Engine/Plugins/MeshPainting/Source/MeshPaintingToolset/Private/MeshPaintHelpers.cpp
+++ b/Engine/Plugins/MeshPainting/Source/MeshPaintingToolset/Private/MeshPaintHelpers.cpp
@@ -101,7 +101,7 @@ void UMeshPaintingSubsystem::CreateComponentMeshPaintTexture(UStaticMeshComponen
 		const uint32 TextureSize = StaticMeshComponent->GetMeshPaintTextureResolution();
 		const uint32 TextureNumMips = FMath::FloorLog2(TextureSize) + 1;
 
-		UMeshPaintVirtualTexture* NewTexture = NewObject<UMeshPaintVirtualTexture>(StaticMeshComponent->GetOutermost());
+		UMeshPaintVirtualTexture* NewTexture = NewObject<UMeshPaintVirtualTexture>(StaticMeshComponent);
 		NewTexture->Source.Init(TextureSize, TextureSize, 1, TextureNumMips, TSF_BGRA8);
 		NewTexture->OwningComponent = MakeWeakObjectPtr(StaticMeshComponent);
 		NewTexture->UpdateResource();
@@ -116,7 +116,7 @@ void UMeshPaintingSubsystem::CreateComponentMeshPaintTexture(UStaticMeshComponen
 	{
 		StaticMeshComponent->Modify();
 
-		UMeshPaintVirtualTexture* NewTexture = NewObject<UMeshPaintVirtualTexture>(StaticMeshComponent->GetOutermost());
+		UMeshPaintVirtualTexture* NewTexture = NewObject<UMeshPaintVirtualTexture>(StaticMeshComponent);
 		NewTexture->Source.Init(InImage);
 		NewTexture->OwningComponent = MakeWeakObjectPtr(StaticMeshComponent);
 		NewTexture->UpdateResource();
@@ -1742,7 +1742,7 @@ TCHAR const* FMeshPaintToolSettingHelpers::GetCacheIdentifier()
 void FMeshPaintToolSettingHelpers::SavePropertiesForClassHeirachy(UInteractiveTool* InTool, UInteractiveToolPropertySet* InProperties)
 {
 	// Properties are saved in a cache on the class CDO. 
-	// So we need to iterate the class structure and save for each class in the hierachy.
+	// So we need to iterate the class structure and save for each class in the hierarchy.
 	// That way classes that share a base class will both use the same cache for their shared properties.
 	// Specifically we keep the same base class brush settings when switching between the different painting tools.
 	const FString CacheIdentifier(GetCacheIdentifier());
@@ -1767,7 +1767,7 @@ void FMeshPaintToolSettingHelpers::SavePropertiesForClassHeirachy(UInteractiveTo
 
 void FMeshPaintToolSettingHelpers::RestorePropertiesForClassHeirachy(UInteractiveTool* InTool, UInteractiveToolPropertySet* InProperties)
 {
-	// Restore from leaf class down the hierachy, with each base class overriding what was written before.
+	// Restore from leaf class down the hierarchy, with each base class overriding what was written before.
 	const FString CacheIdentifier(GetCacheIdentifier());
 	InProperties->RestoreProperties(InTool, CacheIdentifier);
 
diff --git a/Engine/Source/Runtime/Engine/Private/Components/StaticMeshComponent.cpp b/Engine/Source/Runtime/Engine/Private/Components/StaticMeshComponent.cpp
--- a/Engine/Source/Runtime/Engine/Private/Components/StaticMeshComponent.cpp
+++ b/Engine/Source/Runtime/Engine/Private/Components/StaticMeshComponent.cpp
@@ -775,7 +775,7 @@ static UTexture* RecreateMeshPaintTexture(UTexture* InTexture, UPrimitiveCompone
 	FImage Image;
 	if (InTexture->Source.GetMipImage(Image, 0))
 	{
-		NewTexture = NewObject<UMeshPaintVirtualTexture>(InComponent->GetOutermost());
+		NewTexture = NewObject<UMeshPaintVirtualTexture>(InComponent);
 		NewTexture->Source.Init(Image);
 		NewTexture->OwningComponent = MakeWeakObjectPtr(InComponent);
 		NewTexture->UpdateResource();
@@ -3731,7 +3731,7 @@ void FStaticMeshComponentLODInfo::ExportText(FString& ValueStr)
 		ValueStr += FString::Printf(TEXT("(Normal=(X=%f,Y=%f,Z=%f,W=%f),"), Vert.Normal.X, Vert.Normal.Y, Vert.Normal.Z, Vert.Normal.W);
 		ValueStr += FString::Printf(TEXT("(Color=(B=%d,G=%d,R=%d,A=%d))"), Vert.Color.B, Vert.Color.G, Vert.Color.R, Vert.Color.A);
 
-		// Seperate each vertex entry with a comma
+		// Separate each vertex entry with a comma
 		if ((i + 1) != PaintedVertices.Num())
 		{
 			ValueStr += TEXT(",");
