[Crash Fix] Add an additional check on the ObjectInstance as a UActorComponent

#Synced-CL 49764259

diff --git a/Engine/Plugins/MeshPainting/Source/MeshPaintingToolset/Private/MeshPaintInteractions.cpp b/Engine/Plugins/MeshPainting/Source/MeshPaintingToolset/Private/MeshPaintInteractions.cpp
--- a/Engine/Plugins/MeshPainting/Source/MeshPaintingToolset/Private/MeshPaintInteractions.cpp
+++ b/Engine/Plugins/MeshPainting/Source/MeshPaintingToolset/Private/MeshPaintInteractions.cpp
@@ -93,27 +93,30 @@ bool UMeshPaintSelectionMechanic::FindClickedComponentsAndCacheAdapters(const FI
 		{
 			if (TTypedElement<ITypedElementObjectInterface> ObjectInterface = UTypedElementRegistry::GetInstance()->GetElement<ITypedElementObjectInterface>(HitProxy->GetElementHandle()))
 			{
-				if (AActor* Actor = ObjectInterface.GetObjectAs<UActorComponent>()->GetOwner())
+				if (UActorComponent* ObjectInterfaceActorComponent = ObjectInterface.GetObjectAs<UActorComponent>())
 				{
-					TArray<UActorComponent*> CandidateComponents = Actor->K2_GetComponentsByClass(UMeshComponent::StaticClass());
-					for (UActorComponent* CandidateComponent : CandidateComponents)
+					if (AActor* Actor = ObjectInterfaceActorComponent->GetOwner())
 					{
-						UMeshComponent* MeshComponent = Cast<UMeshComponent>(CandidateComponent);
-						TSharedPtr<IMeshPaintComponentAdapter> MeshAdapter = MeshPaintingSubsystem->GetAdapterForComponent(MeshComponent);
-						if (!MeshAdapter.IsValid() || !MeshAdapter->IsValid())
+						TArray<UActorComponent*> CandidateComponents = Actor->K2_GetComponentsByClass(UMeshComponent::StaticClass());
+						for (UActorComponent* CandidateComponent : CandidateComponents)
 						{
-							MeshAdapter = FMeshPaintComponentAdapterFactory::CreateAdapterForMesh(MeshComponent, 0);
-							if (MeshAdapter)
+							UMeshComponent* MeshComponent = Cast<UMeshComponent>(CandidateComponent);
+							TSharedPtr<IMeshPaintComponentAdapter> MeshAdapter = MeshPaintingSubsystem->GetAdapterForComponent(MeshComponent);
+							if (!MeshAdapter.IsValid() || !MeshAdapter->IsValid())
 							{
-								MeshPaintingSubsystem->AddToComponentToAdapterMap(MeshComponent, MeshAdapter);
+								MeshAdapter = FMeshPaintComponentAdapterFactory::CreateAdapterForMesh(MeshComponent, 0);
+								if (MeshAdapter)
+								{
+									MeshPaintingSubsystem->AddToComponentToAdapterMap(MeshComponent, MeshAdapter);
+								}
+							}
+							IMeshPaintSelectionInterface* Interface = Cast<IMeshPaintSelectionInterface>(GetParentTool());
+							if (MeshAdapter.IsValid() && Interface->IsMeshAdapterSupported(MeshAdapter))
+							{
+								CachedClickedComponents.AddUnique(MeshComponent);
+								CachedClickedActors.AddUnique(Cast<AActor>(MeshComponent->GetOuter()));
+								bFoundValidComponents = true;
 							}
-						}
-						IMeshPaintSelectionInterface* Interface = Cast<IMeshPaintSelectionInterface>(GetParentTool());
-						if (MeshAdapter.IsValid() && Interface->IsMeshAdapterSupported(MeshAdapter))
-						{
-							CachedClickedComponents.AddUnique(MeshComponent);
-							CachedClickedActors.AddUnique(Cast<AActor>(MeshComponent->GetOuter()));
-							bFoundValidComponents = true;
 						}
 					}
 				}
