Change default mesh paint erase color from black to transparent.
This is following a request to make alpha behave the same as other paint channels.
#rnx

#Synced-CL 49987149

diff --git a/Engine/Plugins/MeshPainting/Source/MeshPaintingToolset/Public/BaseMeshPaintingToolProperties.h b/Engine/Plugins/MeshPainting/Source/MeshPaintingToolset/Public/BaseMeshPaintingToolProperties.h
--- a/Engine/Plugins/MeshPainting/Source/MeshPaintingToolset/Public/BaseMeshPaintingToolProperties.h
+++ b/Engine/Plugins/MeshPainting/Source/MeshPaintingToolset/Public/BaseMeshPaintingToolProperties.h
@@ -17,7 +17,7 @@ public:
 
 	/** Color used for erasing color painting */
 	UPROPERTY(EditAnywhere, Category = ColorPainting)
-	FLinearColor EraseColor = FLinearColor::Black;
+	FLinearColor EraseColor = FLinearColor::Transparent;
 
 	/** Enables "Flow" painting where paint is continually applied from the brush every tick */
 	UPROPERTY(EditAnywhere, Category = Brush, meta = (DisplayName = "Enable Brush Flow"))
