MeshPaint: Filter materials touched by texture overrides.
Materials texture overrides are stored in their static resources.
So apply MIC overrides to only the parent with a static permutations.
Collect unique set of materials before overriding to prevent duplicating the work.
#rnx
#Synced-CL 50244055

diff --git a/Engine/Plugins/MeshPainting/Source/MeshPaintingToolset/Private/IMeshPaintComponentAdapter.cpp b/Engine/Plugins/MeshPainting/Source/MeshPaintingToolset/Private/IMeshPaintComponentAdapter.cpp
--- a/Engine/Plugins/MeshPainting/Source/MeshPaintingToolset/Private/IMeshPaintComponentAdapter.cpp
+++ b/Engine/Plugins/MeshPainting/Source/MeshPaintingToolset/Private/IMeshPaintComponentAdapter.cpp
@@ -423,9 +423,33 @@ namespace UE::MeshPaintingToolset
 
 		const ERHIFeatureLevel::Type FeatureLevel = InMeshComponent->GetWorld()->GetFeatureLevel();
 
+		TSet<UMaterialInterface*> MaterialsToOverride;
+
 		// Applying the texture override to all materials shows live update throughout the scene not just on painted component.
 		FObjectCacheContextScope ObjectCacheScope;
 		for (UMaterialInterface* Material : ObjectCacheScope.GetContext().GetMaterialsAffectedByTexture(SourceTexture))
+		{
+			// Texture overrides are stored on the static resources so walk to first parent with static permutation.
+			UMaterialInstance* MaterialInstance = Cast<UMaterialInstance>(Material);
+			while (MaterialInstance)
+			{
+				if (MaterialInstance->bHasStaticPermutationResource)
+				{
+					break;
+				}
+				MaterialInstance = Cast<UMaterialInstance>(MaterialInstance->Parent);
+			}
+			if (MaterialInstance)
+			{
+				MaterialsToOverride.Add(MaterialInstance);
+			}
+			else if (UMaterial* BaseMaterial = Material->GetMaterial())
+			{
+				MaterialsToOverride.Add(BaseMaterial);
+			}
+		}
+
+		for (UMaterialInterface* Material : MaterialsToOverride)
 		{
 			bool bMaterialUpdated = false;
 			
