MeshPaint: RVT invalidation on paint reduced to brush area instead of entire RVT bounds.
Reduces visible glitching when painting to RVT materials.
#rnx
#Synced-CL 50294033

diff --git a/Engine/Plugins/MeshPainting/Source/MeshPaintingToolset/Private/MeshTexturePaintingTool.cpp b/Engine/Plugins/MeshPainting/Source/MeshPaintingToolset/Private/MeshTexturePaintingTool.cpp
--- a/Engine/Plugins/MeshPainting/Source/MeshPaintingToolset/Private/MeshTexturePaintingTool.cpp
+++ b/Engine/Plugins/MeshPainting/Source/MeshPaintingToolset/Private/MeshTexturePaintingTool.cpp
@@ -1237,9 +1237,10 @@ void UMeshTexturePaintingTool::PaintTexture(FMeshPaintParameters& InParams, int3
 	// Need to invalidate runtime virtual textures if paint texture is used in a material that writes to one.
 	if (bRequiresRuntimeVirtualTextureUpdates)
 	{
+		FBoxSphereBounds BrushBounds(FSphere(InParams.BrushPosition, InParams.InnerBrushRadius + InParams.BrushRadialFalloffRange));
 		for (TObjectIterator<URuntimeVirtualTextureComponent> It(RF_ClassDefaultObject, true, EInternalObjectFlags::Garbage); It; ++It)
 		{
-			It->Invalidate(It->Bounds, EVTInvalidatePriority::High);
+			It->Invalidate(BrushBounds, EVTInvalidatePriority::High);
 		}
 	}
 }
