[Fix] Color now propagates to inconsistent mesh

[Wrong behaviour] A mesh with the number of wedges not equal to the number of verticies would not be propagated color from the MeshPaintingToolset on "ToMesh" button press.

[Fixed] Added check for "inconsistent" meshes. If the wedge map exists and is inconsistent with the mesh description, it will not use the wedge map to propagate colors, but will use the fallback method.

#14232

#Synced-CL 50447978

diff --git a/Engine/Plugins/MeshPainting/Source/MeshPaintingToolset/Private/MeshPaintHelpers.cpp b/Engine/Plugins/MeshPainting/Source/MeshPaintingToolset/Private/MeshPaintHelpers.cpp
--- a/Engine/Plugins/MeshPainting/Source/MeshPaintingToolset/Private/MeshPaintHelpers.cpp
+++ b/Engine/Plugins/MeshPainting/Source/MeshPaintingToolset/Private/MeshPaintHelpers.cpp
@@ -153,52 +153,47 @@ bool UMeshPaintingSubsystem::PropagateColorsToRawMesh(UStaticMesh* StaticMesh, i
 	FStaticMeshRenderData& RenderData = *StaticMesh->GetRenderData();
 	FStaticMeshLODResources& RenderModel = RenderData.LODResources[LODIndex];
 	FColorVertexBuffer& ColorVertexBuffer = *ComponentLODInfo.OverrideVertexColors;
-	if (RenderModel.WedgeMap.Num() > 0 && ColorVertexBuffer.GetNumVertices() == RenderModel.GetNumVertices())
-	{
-		// Use the wedge map if it is available as it is lossless.
-		FMeshDescription* MeshDescription = StaticMesh->GetMeshDescription(LODIndex);
 
-		if (MeshDescription == nullptr)
-		{
-			//Cannot propagate to a generated LOD, the generated LOD use the source LOD vertex painting. 
-			return false;
-		}
+	FMeshDescription* MeshDescription = StaticMesh->GetMeshDescription(LODIndex);
+
+	if (MeshDescription == nullptr)
+	{
+		return false;
+	}
 
+	// Use the wedge map if it is available as it is lossless.
+	if (RenderModel.WedgeMap.Num() > 0 && ColorVertexBuffer.GetNumVertices() == RenderModel.GetNumVertices() && RenderModel.WedgeMap.Num() == MeshDescription->VertexInstances().Num())
+	{
 		FStaticMeshAttributes Attributes(*MeshDescription);
-		int32 NumWedges = MeshDescription->VertexInstances().Num();
-		if (RenderModel.WedgeMap.Num() == NumWedges)
+		TVertexInstanceAttributesRef<FVector4f> Colors = Attributes.GetVertexInstanceColors();
+
+		int32 VertexInstanceIndex = 0;
+		for (const FVertexInstanceID VertexInstanceID : MeshDescription->VertexInstances().GetElementIDs())
 		{
-			TVertexInstanceAttributesRef<FVector4f> Colors = Attributes.GetVertexInstanceColors();
-			int32 VertexInstanceIndex = 0;
-			for (const FVertexInstanceID VertexInstanceID : MeshDescription->VertexInstances().GetElementIDs())
+			FLinearColor WedgeColor = FLinearColor::White;
+			int32 Index = RenderModel.WedgeMap[VertexInstanceIndex];
+			if (Index != INDEX_NONE)
 			{
-				FLinearColor WedgeColor = FLinearColor::White;
-				int32 Index = RenderModel.WedgeMap[VertexInstanceIndex];
-				if (Index != INDEX_NONE)
-				{
-					WedgeColor = FLinearColor(ColorVertexBuffer.VertexColor(Index));
-				}
-				Colors[VertexInstanceID] = WedgeColor;
-				VertexInstanceIndex++;
+				WedgeColor = FLinearColor(ColorVertexBuffer.VertexColor(Index));
 			}
-			StaticMesh->CommitMeshDescription(LODIndex);
-			bPropagatedColors = true;
+			Colors[VertexInstanceID] = WedgeColor;
+			VertexInstanceIndex++;
 		}
+		StaticMesh->CommitMeshDescription(LODIndex);
+		bPropagatedColors = true;
 	}
-	else
+	else // Fall back to mapping based on position.
 	{
-		FMeshDescription* MeshDescription = StaticMesh->GetMeshDescription(LODIndex);
 		// If there's no raw mesh data, don't try to do any fixup here
-		if (MeshDescription == nullptr || ComponentLODInfo.OverrideVertexColors == nullptr)
+		if (ComponentLODInfo.OverrideVertexColors == nullptr)
 		{
 			return false;
 		}
 
 		FStaticMeshAttributes Attributes(*MeshDescription);
 
-		// Fall back to mapping based on position.
-		TVertexAttributesConstRef<FVector3f> VertexPositions = Attributes.GetVertexPositions();
 		TVertexInstanceAttributesRef<FVector4f> Colors = Attributes.GetVertexInstanceColors();
+		TVertexAttributesConstRef<FVector3f> VertexPositions = Attributes.GetVertexPositions();
 		TArray<FColor> NewVertexColors;
 		FPositionVertexBuffer TempPositionVertexBuffer;
 		int32 NumVertex = MeshDescription->Vertices().Num();
