make the mesh paint mode render geometry collections w/ the native render, so it does not show any proxy static mesh, so that users can see what they're painting

#Synced-CL 51421910

diff --git a/Engine/Plugins/MeshPainting/Source/MeshPaintingToolset/Private/MeshPaintGeometryCollectionAdapter.cpp b/Engine/Plugins/MeshPainting/Source/MeshPaintingToolset/Private/MeshPaintGeometryCollectionAdapter.cpp
--- a/Engine/Plugins/MeshPainting/Source/MeshPaintingToolset/Private/MeshPaintGeometryCollectionAdapter.cpp
+++ b/Engine/Plugins/MeshPainting/Source/MeshPaintingToolset/Private/MeshPaintGeometryCollectionAdapter.cpp
@@ -144,6 +144,7 @@ void FMeshPaintGeometryCollectionComponentAdapter::OnAdded()
 	}
 	bSavedShowBoneColors = GeometryCollectionComponent->GetShowBoneColors();
 	GeometryCollectionComponent->SetShowBoneColors(false);
+	GeometryCollectionComponent->ForceNativeRendering(true);
 
 	checkf(GetGeometryCollectionObject(), TEXT("Geometry Collection Component did not have a valid geometry collection object attached"));
 	checkf(GetGeometryCollectionObject()->GetGeometryCollection(), TEXT("Geometry Collection Component did not have valid geometry collection data attached"));
@@ -155,6 +156,7 @@ void FMeshPaintGeometryCollectionComponentAdapter::OnRemoved()
 	{
 		GeometryCollectionComponent->SetShowBoneColors(bSavedShowBoneColors);
 		bSavedShowBoneColors = false;
+		GeometryCollectionComponent->ForceNativeRendering(false);
 	}
 }
 
